Moving TrajectoryTools plugin to Experimental.

[RN] minor

diff --git a/Engine/Plugins/Experimental/Animation/TrajectoryTools/Source/TrajectoryTools/Private/SExportTrajectoriesWindow.cpp b/Engine/Plugins/Experimental/Animation/TrajectoryTools/Source/TrajectoryTools/Private/SExportTrajectoriesWindow.cpp
new file mode 100644
--- /dev/null
+++ b/Engine/Plugins/Experimental/Animation/TrajectoryTools/Source/TrajectoryTools/Private/SExportTrajectoriesWindow.cpp
@@ -0,0 +1,1141 @@
+// Copyright Epic Games, Inc. All Rights Reserved.
+
+#include "SExportTrajectoriesWindow.h"
+
+#include "PropertyCustomizationHelpers.h"
+#include "SSimpleTimeSlider.h"
+#include "Widgets/Input/SComboButton.h"
+#include "Widgets/Input/SMenuAnchor.h"
+#include "Widgets/Input/SButton.h"
+#include "Widgets/Layout/SBorder.h"
+#include "Layout/WidgetPath.h"
+#include "Widgets/SBoxPanel.h"
+#include "ToolMenu.h"
+#include "ToolMenus.h"
+#include "Widgets/Views/SListView.h"
+#include "Widgets/Views/ITableRow.h"
+#include "Animation/AnimSequence.h"
+#include "Viewports.h"
+#include "Widgets/Text/SRichTextBlock.h"
+#include "EditorModeManager.h"
+#include "Framework/Application/SlateApplication.h"
+#include "Engine/Font.h"
+#include "CanvasTypes.h"
+#include "SceneView.h"
+#include "Animation/DebugSkelMeshComponent.h"
+#include "Animation/Skeleton.h"
+#include "TrajectoryExportOperation.h"
+#include "Fonts/FontMeasure.h"
+#include "Widgets/Layout/SBorder.h"
+#include "Logging/StructuredLog.h"
+
+#define LOCTEXT_NAMESPACE "SExportTrajectoriesWindow"
+
+int32 GetSampleIndexFromTime(const FFrameRate& InFrameRate, double InTime, int32 InSampleNum)
+{
+	return FMath::Clamp(InFrameRate.AsFrameTime(InTime).FloorToFrame().Value, 0, InSampleNum - 1);
+}
+
+void PreviewViewportFocusOnTrajectory(const TSharedPtr<FPreviewData>& InPreviewData, const TSharedPtr<FEditorViewportClient>& InViewportClient)
+{
+	if (!InPreviewData || !InViewportClient)
+	{
+		return;
+	}
+
+	// Build trajectory bounds
+	FBoxSphereBounds::Builder BoundsBuilder;
+	for (const FGameplayTrajectory::FSample& OutputSample : InPreviewData->OutputTrajectory.Samples)
+    {
+		BoundsBuilder += OutputSample.Position;
+    }
+	
+	const FBoxSphereBounds Bound = BoundsBuilder.IsValid() ? BoundsBuilder : FBoxSphereBounds(FSphere(FVector::ZeroVector, 250.0f));
+
+	// Focus
+	InViewportClient->SetViewLocation(Bound.Origin);
+	InViewportClient->SetViewLocationForOrbiting(Bound.Origin, Bound.SphereRadius * 1.5f);
+}
+
+void SPreviewTrajectoryViewport::Construct(const FArguments& InArgs)
+{
+	SEditorViewport::Construct(SEditorViewport::FArguments()
+		.IsEnabled(FSlateApplication::Get().GetNormalExecutionAttribute())
+		.AddMetaData<FTagMetaData>(TEXT("TrajectoryExport.Viewport")));
+
+	PreviewData = InArgs._PreviewData;
+
+	PreviewComponent = NewObject<UDebugSkelMeshComponent>();
+
+	// Always refresh pose and ignore root motion (as position is dictated by trajectory)
+	PreviewComponent->VisibilityBasedAnimTickOption = EVisibilityBasedAnimTickOption::AlwaysTickPoseAndRefreshBones;
+	PreviewComponent->SetProcessRootMotionMode(EProcessRootMotionMode::Ignore);
+
+	// No need to tick the preview component as we will me manually force it to.
+	PreviewComponent->PrimaryComponentTick.bStartWithTickEnabled = false;
+	PreviewComponent->PrimaryComponentTick.bCanEverTick = false;
+
+	// No anim instance.
+	PreviewComponent->SetAnimationMode(EAnimationMode::AnimationCustomMode);
+
+	// Ensure preview component gets ticked.
+	PreviewScene.AddComponent(PreviewComponent, FTransform::Identity);
+}
+
+FReply SPreviewTrajectoryViewport::OnKeyDown(const FGeometry& MyGeometry, const FKeyEvent& InKeyEvent)
+{
+	if (InKeyEvent.GetKey() == EKeys::F)
+	{
+		PreviewViewportFocusOnTrajectory(PreviewData, GetViewportClient());
+		return FReply::Handled();
+	}
+
+	return FReply::Unhandled();
+}
+
+void SPreviewTrajectoryViewport::ResetPreviewSkeletalMesh() const
+{
+	// Exit on invalid state.
+	const bool bHasNoSelectedTrajectory = PreviewData->SourceIndex == INDEX_NONE;
+	if (bHasNoSelectedTrajectory || !IRewindDebugger::Instance() || !IRewindDebugger::Instance()->GetAnalysisSession())
+	{
+		return;
+	}
+
+	// Get data providers for traced objets.
+	IRewindDebugger* RewindDebugger = IRewindDebugger::Instance();
+	const TraceServices::IAnalysisSession* Session = RewindDebugger->GetAnalysisSession();
+	const IGameplayProvider* GameplayProvider = Session->ReadProvider<IGameplayProvider>("GameplayProvider");
+	const IAnimationProvider* AnimationProvider = Session->ReadProvider<IAnimationProvider>("AnimationProvider");
+
+	// Preview traced skeletal mesh.
+	if (GameplayProvider && AnimationProvider)
+	{
+		TraceServices::FAnalysisSessionReadScope AnalysisSessionReadScope(*Session);
+
+		FSkeletalMeshInfo SkeletalMeshInfoToPreview {};
+		{
+			const int32 ScrubSampleIndex = GetSampleIndexFromTime(PreviewData->Details->ExportSettings.FrameRate, PreviewData->Timeline.ScrubTime, PreviewData->OutputTrajectory.Samples.Num());
+		
+			for (int i = 0; i < PreviewData->OutputTrajectory.TraceInfo.Num(); ++i)
+			{
+				if (PreviewData->OutputTrajectory.TraceInfo.Ranges[i].Contains(ScrubSampleIndex))
+				{
+					SkeletalMeshInfoToPreview = PreviewData->OutputTrajectory.TraceInfo.SkeletalMeshInfos[i];
+					break;
+				}
+			}
+		}
+		ensureMsgf(SkeletalMeshInfoToPreview.Id != 0, TEXT("Invalid skeletal mesh to preview. If this gets hit something went wrong."));
+		
+		if (const FSkeletalMeshInfo* SkeletalMeshInfo = AnimationProvider->FindSkeletalMeshInfo(SkeletalMeshInfoToPreview.Id /*PreviewData->Viewport.SourceSkeletalMeshInfo.Id*/))
+		{
+			if (const FObjectInfo* SkeletalMeshObjectInfo = GameplayProvider->FindObjectInfo(SkeletalMeshInfoToPreview.Id /*PreviewData->Viewport.SourceSkeletalMeshInfo.Id)*/))
+			{
+				USkeletalMesh* SkeletalMesh = TSoftObjectPtr<USkeletalMesh>(FSoftObjectPath(SkeletalMeshObjectInfo->PathName)).LoadSynchronous();
+
+				// Try the skeleton asset preview mesh as a fallback when skeletal mesh asset is not valid.
+				if (SkeletalMesh == nullptr)
+				{
+					// Ensure our skeleton is valid.
+					if (SkeletalMeshInfo->SkeletonId != 0)
+					{
+						const FObjectInfo& SkeletonObjectInfo = GameplayProvider->GetObjectInfo(SkeletalMeshInfo->SkeletonId);
+					
+						USkeleton* Skeleton = TSoftObjectPtr<USkeleton>(FSoftObjectPath(SkeletonObjectInfo.PathName)).LoadSynchronous();
+						if (Skeleton)
+						{
+							SkeletalMesh = Skeleton->GetPreviewMesh(true);
+						}
+					}
+				}
+
+				// Assign loaded skeletal mesh.
+				if (SkeletalMesh)
+				{
+					PreviewComponent->SetSkeletalMesh(SkeletalMesh);
+				}
+			}
+		}
+
+		// Run method that sets world position and bone transforms.
+		UpdatePreviewPoseFromScrubTime(AnimationProvider, GameplayProvider, 0, RewindDebugger->CurrentTraceTime());
+	}
+}
+
+void SPreviewTrajectoryViewport::UpdatePreviewPoseFromScrubTime(const IAnimationProvider* InAnimationProvider, const IGameplayProvider* InGameplayProvider, double InTraceStartTime, double InTraceEndTime) const
+{
+	if (PreviewComponent->GetSkeletalMeshAsset() == nullptr || PreviewData->SourceIndex == INDEX_NONE)
+	{
+		return;
+	}
+
+	if (PreviewData->SourceTrajectories[PreviewData->SourceIndex].Samples.IsEmpty() || PreviewData->SourceTrajectories[PreviewData->SourceIndex].Poses.IsEmpty())
+	{
+		return;
+	}
+	
+	const int32 ScrubSampleIndex = GetSampleIndexFromTime(PreviewData->Details->ExportSettings.FrameRate, PreviewData->Timeline.ScrubTime, PreviewData->OutputTrajectory.Samples.Num());
+	const double ScrubSampleTime = PreviewData->OutputTrajectory.Samples[ScrubSampleIndex].Time;
+	const FVector ScrubSamplePosition = PreviewData->OutputTrajectory.Samples[ScrubSampleIndex].Position;
+	const FQuat ScrubSampleOrientation = PreviewData->OutputTrajectory.Samples[ScrubSampleIndex].Orientation;
+
+	// Update transform first, in case, we aren't able to update the pose itself.
+	PreviewComponent->SetWorldTransform(FTransform(ScrubSampleOrientation, ScrubSamplePosition), false, nullptr, ETeleportType::TeleportPhysics);
+
+	// Get preview pose buffer.
+	TArray<FTransform>& EditablePose = PreviewComponent->GetEditableComponentSpaceTransforms();
+	if (EditablePose.Num() == PreviewData->OutputTrajectory.Poses[ScrubSampleIndex].Num())
+	{
+		EditablePose = PreviewData->OutputTrajectory.Poses[ScrubSampleIndex];
+	}
+	else
+	{
+		UE_LOGFMT(LogTemp, Warning, "Preview Skeletal Mesh Cmp transform buffer does not match the recorded sample pose buffer size. {num1} vs {num2}", EditablePose.Num(), PreviewData->OutputTrajectory.Poses[ScrubSampleIndex].Num());
+		return;
+	}
+
+	// Update current preview pose.
+	FTrajectoryToolsLibrary::GetPoseAtTimeInTrajectory(PreviewData->SourceTrajectories[PreviewData->SourceIndex], ScrubSampleTime, EditablePose);
+	
+	// Apply new transforms.
+	PreviewComponent->ApplyEditedComponentSpaceTransforms();
+
+	// @todo: Handle LOD setting.
+	// PreviewComponent->SetForcedLOD(PoseMessage->LodIndex + 1);
+
+	// Update preview component's properties.
+	PreviewComponent->UpdateLODStatus();
+	PreviewComponent->UpdateChildTransforms(EUpdateTransformFlags::None, ETeleportType::TeleportPhysics);
+	PreviewComponent->SetVisibility(true);
+	PreviewComponent->MarkRenderStateDirty();
+	PreviewComponent->SetDrawDebugSkeleton(true);
+}
+
+TSharedRef<FEditorViewportClient> SPreviewTrajectoryViewport::MakeEditorViewportClient()
+{
+	TSharedPtr<FEditorViewportClient> EditorViewportClient = MakeShareable(new FPreviewTrajectoryViewportClient(PreviewScene, SharedThis(this)));
+	return EditorViewportClient.ToSharedRef();
+}
+
+TSharedPtr<SWidget> SPreviewTrajectoryViewport::MakeViewportToolbar()
+{
+	// Use an empty toolbar
+	return SEditorViewport::MakeViewportToolbar();
+}
+
+FPreviewTrajectoryViewportClient::FPreviewTrajectoryViewportClient(FPreviewScene& InPreviewScene, const TSharedRef<SPreviewTrajectoryViewport>& InPreviewTrajectoryViewport)
+	: FEditorViewportClient(nullptr, &InPreviewScene, StaticCastSharedRef<SEditorViewport>(InPreviewTrajectoryViewport))
+{
+	ViewportType = LVT_Perspective;
+	bSetListenerPosition = false;
+
+	// Use defaults for view transforms.
+	SetViewLocation(EditorViewportDefs::DefaultPerspectiveViewLocation);
+	SetViewRotation(EditorViewportDefs::DefaultPerspectiveViewRotation);
+	SetViewLocationForOrbiting(FVector::ZeroVector);
+
+	// Normally the bIsRealtime flag is determined by whether the connection is remote, but our
+	// tools require always being ticked.
+	SetRealtime(true);
+
+	// Lit gives us the most options in terms of the materials we can use.
+	SetViewMode(VMI_Lit);
+	
+	// Allow for camera control.
+	bUsingOrbitCamera = true;
+	bDisableInput = false;
+	
+	// This seems to be needed to get the correct world time in the preview.
+	// SetIsSimulateInEditorViewport(true);
+	
+	// Always composite editor objects after post-processing in the editor
+	EngineShowFlags.SetCompositeEditorPrimitives(true);
+	EngineShowFlags.DisableAdvancedFeatures();
+	
+	// Setup defaults for the common draw helper.
+	DrawHelper.bDrawPivot = false;
+	DrawHelper.bDrawWorldBox = false;
+	DrawHelper.bDrawKillZ = false;
+	DrawHelper.bDrawGrid = true;
+	DrawHelper.GridColorAxis = FColor(40, 40, 40);
+	DrawHelper.GridColorMajor = FColor(20, 20, 20);
+	DrawHelper.GridColorMinor =  FColor(10, 10, 10);
+	DrawHelper.PerspectiveGridSize = UE_OLD_HALF_WORLD_MAX1;
+}
+
+void FPreviewTrajectoryViewportClient::Tick(float DeltaTime)
+{
+	FEditorViewportClient::Tick(DeltaTime);
+	
+	// Used to update preview pose.
+	TWeakPtr<SPreviewTrajectoryViewport> PreviewEditorViewport = StaticCastWeakPtr<SPreviewTrajectoryViewport>(EditorViewportWidget);
+
+	if (PreviewEditorViewport.IsValid())
+	{
+		if (IRewindDebugger* RewindDebugger = IRewindDebugger::Instance())
+		{
+			// Get data providers for traced objets.
+			if (const TraceServices::IAnalysisSession* Session = RewindDebugger->GetAnalysisSession())
+			{
+				const IGameplayProvider* GameplayProvider = Session->ReadProvider<IGameplayProvider>("GameplayProvider");
+				const IAnimationProvider* AnimationProvider = Session->ReadProvider<IAnimationProvider>("AnimationProvider");
+
+				// Preview traced skeletal mesh.
+				if (GameplayProvider && AnimationProvider)
+				{
+					// Start reading this trace time information.
+					TraceServices::FAnalysisSessionReadScope AnalysisSessionReadScope(*Session);
+					PreviewEditorViewport.Pin().Get()->UpdatePreviewPoseFromScrubTime(AnimationProvider, GameplayProvider, 0, RewindDebugger->CurrentTraceTime());
+				}
+			}	
+		}
+	}
+
+	// Tick preview components.
+	if (PreviewScene)
+	{
+		PreviewScene->GetWorld()->Tick(LEVELTICK_All, DeltaTime);
+	}
+}
+
+void FPreviewTrajectoryViewportClient::Draw(const FSceneView* InView, FPrimitiveDrawInterface* InPDI)
+{
+	constexpr float DepthBias = 4.0f;
+	constexpr bool bScreenSpace = true;
+	
+	FEditorViewportClient::Draw(InView, InPDI);
+
+	if (const TSharedPtr<SPreviewTrajectoryViewport> ViewportWidget = StaticCastSharedPtr<SPreviewTrajectoryViewport>(EditorViewportWidget.Pin()))
+	{
+		if (const TSharedPtr<FPreviewData> PreviewData = ViewportWidget->GetPreviewData())
+		{
+			if (!PreviewData->OutputTrajectory.Samples.IsEmpty())
+			{
+				// Are we exporting a segment of the trajectory
+				const FFloatInterval& ExportRange = PreviewData->Details->ExportSettings.Range;
+				const bool bShouldDrawEntireTrajectory = FMath::IsNearlyZero(ExportRange.Size()) /*|| !ExportRange.IsValid()*/;
+
+				// Draw trajectory
+				for (int SampleIndex = 0; SampleIndex < PreviewData->OutputTrajectory.Samples.Num(); ++SampleIndex)
+				{
+					// Current sample info
+					const FGameplayTrajectory::FSample& Sample = PreviewData->OutputTrajectory.Samples[SampleIndex];
+					const bool bIsSampleInRange = ExportRange.Contains(Sample.Time);
+					FTransform SampleTransform(Sample.Orientation, Sample.Position);
+					const uint8 SampleAlpha = bShouldDrawEntireTrajectory ? 255 : (bIsSampleInRange ? 255 : 25);
+					FLinearColor SampleColor = FColor::Black.WithAlpha(SampleAlpha);
+					
+					// Draw vertical ticks indicating start and end locations.
+					constexpr double SampleTickSize = 5.0;
+					InPDI->DrawTranslucentLine(Sample.Position, Sample.Position + SampleTransform.GetUnitAxis(EAxis::Z) * SampleTickSize, SampleColor, SDPG_World, 1.0f, DepthBias, bScreenSpace);
+
+					// Draw line connecting current sample to next sample
+					if (SampleIndex + 1 < PreviewData->OutputTrajectory.Samples.Num())
+					{
+						const FGameplayTrajectory::FSample& FutureSample = PreviewData->OutputTrajectory.Samples[SampleIndex + 1];
+						const bool bIsFutureSampleInRange = ExportRange.Contains(FutureSample.Time);
+						const uint8 FutureSampleAlpha = bIsSampleInRange && bIsFutureSampleInRange ? 255 : 25;
+						FLinearColor FutureSampleColor = FColor::Black.WithAlpha(FutureSampleAlpha);
+
+						InPDI->DrawTranslucentLine(Sample.Position, FutureSample.Position, FutureSampleColor, SDPG_World, 1.5f, DepthBias, bScreenSpace);
+					}
+				}
+
+				// Draw scrubber point
+				if (PreviewData->Timeline.bShouldDrawScrubber)
+				{
+					const int32 ScrubSampleIndex = GetSampleIndexFromTime(PreviewData->Details->ExportSettings.FrameRate, PreviewData->Timeline.ScrubTime, PreviewData->OutputTrajectory.Samples.Num());
+					
+					InPDI->DrawPoint(PreviewData->OutputTrajectory.Samples[ScrubSampleIndex].Position, FColor::White, 7.0f, SDPG_Foreground);
+				}
+
+				// Draw start and end export range points
+				if (PreviewData->Details->ExportSettings.bShouldForceOrigin)
+				{
+					const int32 OriginSampleIndex = GetSampleIndexFromTime(PreviewData->Details->ExportSettings.FrameRate, PreviewData->Details->ExportSettings.OriginTime,PreviewData->OutputTrajectory.Samples.Num());
+					
+					InPDI->DrawPoint(PreviewData->OutputTrajectory.Samples[OriginSampleIndex].Position, FColor::Blue, 7.0f, SDPG_Foreground);
+				}
+			}
+		}
+	}
+}
+
+void FPreviewTrajectoryViewportClient::DrawCanvas(FViewport& InViewport, FSceneView& InView, FCanvas& InCanvas)
+{
+	FEditorViewportClient::DrawCanvas(InViewport, InView, InCanvas);
+	
+	if (const TSharedPtr<SPreviewTrajectoryViewport> ViewportWidget = StaticCastSharedPtr<SPreviewTrajectoryViewport>(EditorViewportWidget.Pin()))
+	{
+		if (const TSharedPtr<FPreviewData> PreviewData = ViewportWidget->GetPreviewData())
+		{
+			if (!PreviewData->OutputTrajectory.Samples.IsEmpty())
+			{
+				const auto DrawInfoStringForSample = [](FSceneView& InView, FCanvas& InCanvas, uint32 InSampleIndex, double InSampleTime, const FVector& InSamplePosition)
+				{
+					FVector2D PixelLocation;
+					if (InView.WorldToPixel(InSamplePosition,PixelLocation))
+					{
+						PixelLocation.X = FMath::RoundToFloat(PixelLocation.X);
+						PixelLocation.Y = FMath::RoundToFloat(PixelLocation.Y);
+
+						constexpr FColor LabelColor(200, 200, 200);
+						constexpr FLinearColor ShadowColor(0, 0, 0, 0.3f);
+						const UFont* Font = GEngine->GetLargeFont();
+
+						PixelLocation.Y -= Font->GetMaxCharHeight() * 2;
+						InCanvas.DrawShadowedString(PixelLocation.X, PixelLocation.Y, FString::Format(TEXT("{0}* ({1}s)"), {InSampleIndex, InSampleTime}), Font, LabelColor, ShadowColor);
+					}
+				};
+
+				// Draw scrubber point
+				if (PreviewData->Timeline.bShouldDrawScrubber)
+				{
+					const int32 ScrubSampleIndex = GetSampleIndexFromTime(PreviewData->Details->ExportSettings.FrameRate, PreviewData->Timeline.ScrubTime, PreviewData->OutputTrajectory.Samples.Num());
+					const double ScrubSampleTime = PreviewData->OutputTrajectory.Samples[ScrubSampleIndex].Time;
+					const FVector& ScrubSamplePosition = PreviewData->OutputTrajectory.Samples[ScrubSampleIndex].Position;
+
+					DrawInfoStringForSample(InView, InCanvas, ScrubSampleIndex, ScrubSampleTime,  ScrubSamplePosition);
+				}
+
+				// Draw origin point
+				if (PreviewData->Details->ExportSettings.bShouldForceOrigin)
+				{
+					const int32 OriginSampleIndex = GetSampleIndexFromTime(PreviewData->Details->ExportSettings.FrameRate, PreviewData->Details->ExportSettings.OriginTime,PreviewData->OutputTrajectory.Samples.Num());
+					const double OriginSampleTime = PreviewData->OutputTrajectory.Samples[OriginSampleIndex].Time;
+					const FVector& OriginSamplePosition = PreviewData->OutputTrajectory.Samples[OriginSampleIndex].Position;
+					
+					DrawInfoStringForSample(InView, InCanvas, OriginSampleIndex, OriginSampleTime,  OriginSamplePosition);
+				}
+				
+				// Draw start and end export range points
+				const FFloatInterval& ExportRange = PreviewData->Details->ExportSettings.Range;
+				const bool bDisplayBoundInfo = !FMath::IsNearlyZero(ExportRange.Size()) && ExportRange.IsValid();
+				if (bDisplayBoundInfo)
+				{
+					const int32 LowerBoundSampleIndex = GetSampleIndexFromTime(PreviewData->Details->ExportSettings.FrameRate, PreviewData->Details->ExportSettings.Range.Min, PreviewData->OutputTrajectory.Samples.Num());
+					const double LowerBoundSampleTime = PreviewData->OutputTrajectory.Samples[LowerBoundSampleIndex].Time;
+					const FVector& LowerBoundSamplePosition = PreviewData->OutputTrajectory.Samples[LowerBoundSampleIndex].Position;
+					
+					const int32 UpperBoundSampleIndex = GetSampleIndexFromTime(PreviewData->Details->ExportSettings.FrameRate, PreviewData->Details->ExportSettings.Range.Max, PreviewData->OutputTrajectory.Samples.Num());
+					const double UpperBoundSampleTime = PreviewData->OutputTrajectory.Samples[UpperBoundSampleIndex].Time;
+					const FVector& UpperBoundSamplePosition = PreviewData->OutputTrajectory.Samples[UpperBoundSampleIndex].Position;
+
+					DrawInfoStringForSample(InView, InCanvas, LowerBoundSampleIndex, LowerBoundSampleTime,  LowerBoundSamplePosition);
+					DrawInfoStringForSample(InView, InCanvas, UpperBoundSampleIndex, UpperBoundSampleTime,  UpperBoundSamplePosition);
+				}
+			}
+		}
+	}
+}
+
+/** Used to preview a point in time of a given trajectory */
+class STrajectoryTimeline : public SCompoundWidget
+{
+public:
+	SLATE_BEGIN_ARGS(STrajectoryTimeline) {}
+		SLATE_ARGUMENT(TSharedPtr<FPreviewData>, PreviewData)
+		SLATE_EVENT(FOnContextMenuOpening, OnContextMenuOpening)
+		SLATE_EVENT(SSimpleTimeSlider::FOnScrubPositionChanged, OnScrubPositionChanged)
+	SLATE_END_ARGS()
+
+	void Construct(const FArguments& InArgs)
+	{
+		PreviewData = InArgs._PreviewData;
+		OnContextMenuOpening = InArgs._OnContextMenuOpening;
+		OnScrubPositionChanged = InArgs._OnScrubPositionChanged;
+		
+		ChildSlot
+		[
+			SNew(SOverlay)
+			.Visibility(EVisibility::SelfHitTestInvisible)
+			
+			+ SOverlay::Slot()
+			.VAlign(VAlign_Center)
+			[
+				SNew(STextBlock)
+				.Visibility(EVisibility::HitTestInvisible)
+				.Text(LOCTEXT("NoTrajectoryInTimelineLabel", "No trajectory selected or available"))
+				.Justification(ETextJustify::Center)
+				.Visibility_Lambda([this](){ return PreviewData->OutputTrajectory.Samples.IsEmpty() ? EVisibility::Visible : EVisibility::Collapsed; })
+			]
+			
+			+ SOverlay::Slot()
+			[
+				SAssignNew(Timeline, SSimpleTimeSlider)
+				.IsEnabled_Lambda([this](){ return !PreviewData->OutputTrajectory.Samples.IsEmpty(); })
+				.DesiredSize({100,24})
+				.ClampRangeHighlightSize(0.15f)
+				.ClampRangeHighlightColor(FLinearColor::Red.CopyWithNewOpacity(0.5f))
+				.ScrubPosition_Lambda([this](){ return PreviewData->Timeline.ScrubTime; })
+				.ViewRange_Lambda([this]()
+				{
+					TRange<double> OutputRange(0, 0);
+
+					if (!PreviewData->OutputTrajectory.Samples.IsEmpty())
+					{
+						OutputRange.SetLowerBoundValue(PreviewData->OutputTrajectory.Samples[0].Time);
+						OutputRange.SetUpperBoundValue(PreviewData->OutputTrajectory.Samples.Last().Time);
+					}
+					
+					return OutputRange;
+				})
+				.ClampRange_Lambda([this]()
+				{
+					TRange<double> OutputRange(0, 0);
+
+					if (!PreviewData->OutputTrajectory.Samples.IsEmpty())
+					{
+						OutputRange.SetLowerBoundValue(PreviewData->OutputTrajectory.Samples[0].Time);
+						OutputRange.SetUpperBoundValue(PreviewData->OutputTrajectory.Samples.Last().Time);
+					}
+					
+					return OutputRange;
+				})
+				.OnEndScrubberMovement_Lambda([this]()
+				{
+					PreviewData->Timeline.bShouldDrawScrubber = Timeline->IsHovered();
+				})
+				.OnScrubPositionChanged_Lambda([this](double NewScrubTime, bool bIsScrubbing)
+				{
+					if (bIsScrubbing)
+					{
+						PreviewData->Timeline.ScrubTime = NewScrubTime;
+					}
+
+					PreviewData->Timeline.bShouldDrawScrubber = true;
+					
+					OnScrubPositionChanged.ExecuteIfBound(NewScrubTime, bIsScrubbing);
+				})
+			]
+		];
+	}
+
+	/** Begin SWidget */
+	virtual void OnMouseEnter(const FGeometry& MyGeometry, const FPointerEvent& MouseEvent) override
+	{
+		PreviewData->Timeline.bShouldDrawScrubber = true;
+	}
+		
+	virtual void OnMouseLeave(const FPointerEvent& MouseEvent) override
+	{
+		PreviewData->Timeline.bShouldDrawScrubber = false;
+	}
+
+	virtual FReply OnKeyDown(const FGeometry& MyGeometry, const FKeyEvent& InKeyEvent) override
+	{
+		//@todo: Look into why this doesnt work?
+		const bool bIsRightKey = InKeyEvent.GetKey() == EKeys::Right;
+		const bool bIsLeftKey = InKeyEvent.GetKey() == EKeys::Left;
+		
+		if ((bIsRightKey || bIsLeftKey) && IsValid(PreviewData->Details))
+		{
+			const float Multiplier = static_cast<float>(bIsRightKey) - static_cast<float>(bIsLeftKey);
+			PreviewData->Timeline.ScrubTime += PreviewData->Details->ExportSettings.FrameRate.AsDecimal() * Multiplier;
+			
+			return FReply::Handled();
+		}
+
+		return FReply::Unhandled();
+	}
+
+	virtual FReply OnPreviewMouseButtonDown(const FGeometry& MyGeometry, const FPointerEvent& InMouseEvent) override
+	{
+		const FVector2f SummonLocation = InMouseEvent.GetScreenSpacePosition();
+		const bool bIsRightMouseButtonDown = InMouseEvent.GetEffectingButton() == EKeys::RightMouseButton;
+		const bool bContextMenuOpeningBound = OnContextMenuOpening.IsBound();
+		
+		if (bIsRightMouseButtonDown)
+		{
+			if (bContextMenuOpeningBound)
+			{
+				// Get the context menu content. If NULL, don't open a menu.
+				const TSharedPtr<SWidget> MenuContent = OnContextMenuOpening.Execute();
+
+				if (MenuContent.IsValid())
+				{
+					const FWidgetPath WidgetPath = InMouseEvent.GetEventPath() != nullptr ? *InMouseEvent.GetEventPath() : FWidgetPath();
+					FSlateApplication::Get().PushMenu(AsShared(), WidgetPath, MenuContent.ToSharedRef(), SummonLocation, FPopupTransitionEffect(FPopupTransitionEffect::ContextMenu));
+				}
+			}
+
+			// Note that we intentionally not handle the event so that the SSimpleTimeSlider is still able to handle its OnPreviewMouseButtonDown().
+		}
+
+		return FReply::Unhandled();
+	}
+
+	static void PaintBoneSection(const FGeometry& AllottedGeometry, const FSlateFontInfo& InFont, FSlateWindowElementList& OutDrawElements, int32 InLayerId, const SSimpleTimeSlider::FScrubRangeToScreen& InRangeToScreen, uint32 InBoneCount, float InStartTime, float InEndTime, int InSectionId)
+	{
+		float LeftSection = InRangeToScreen.InputToLocalX(InStartTime);
+		float RightSection = InRangeToScreen.InputToLocalX(InEndTime);
+		
+		float SectionSize = RightSection - LeftSection;
+		float SectionHorizontalPosition = LeftSection;
+
+		FString SectionString = FString::Format(TEXT("Bone Count: {0}"), { InBoneCount });
+		TSharedRef<FSlateFontMeasure> FontMeasureService = FSlateApplication::Get().GetRenderer()->GetFontMeasureService();
+		const FVector2D TextSize = FontMeasureService->Measure(SectionString, InFont);
+		FVector2f TextOffset = FVector2f(SectionSize * 0.5f - TextSize.X * 0.5f);
+		FLinearColor BackgroundColor = (FLinearColor::MakeRandomSeededColor(1999 + InSectionId) * 0.3f).ToFColorSRGB().WithAlpha(255);
+		float Height = AllottedGeometry.GetLocalSize().Y * 0.5f;
+		
+		FPaintGeometry RangeGeometry;
+		RangeGeometry = AllottedGeometry.ToPaintGeometry(FVector2f(SectionSize, Height), FSlateLayoutTransform(FVector2f(SectionHorizontalPosition, -AllottedGeometry.GetLocalSize().Y + Height)));
+		FSlateDrawElement::MakeBox(
+			OutDrawElements,
+			InLayerId + 1,
+			RangeGeometry,
+			FAppStyle::GetBrush("WhiteBrush"),
+			ESlateDrawEffect::None,
+			BackgroundColor
+			);
+
+		RangeGeometry = AllottedGeometry.ToPaintGeometry(FVector2f(SectionSize, Height), FSlateLayoutTransform(FVector2f(SectionHorizontalPosition + TextOffset.X, -AllottedGeometry.GetLocalSize().Y + Height)));
+		FSlateDrawElement::MakeText(
+			OutDrawElements,
+			InLayerId + 2,
+			RangeGeometry,
+			SectionString,
+			InFont,
+			ESlateDrawEffect::None,
+			FLinearColor::White
+		);
+	}
+	
+	virtual int32 OnPaint(const FPaintArgs& Args, const FGeometry& AllottedGeometry, const FSlateRect& MyCullingRect, FSlateWindowElementList& OutDrawElements, int32 LayerId, const FWidgetStyle& InWidgetStyle, bool bParentEnabled) const override
+	{
+		// Draw normal widget.
+		LayerId = SCompoundWidget::OnPaint(Args, AllottedGeometry, MyCullingRect, OutDrawElements, LayerId, InWidgetStyle, bParentEnabled);
+
+		// Draw bone count splits.
+		if (PreviewData->SourceIndex != INDEX_NONE)
+		{
+			const FGameplayTrajectory SourceTrajectory = PreviewData->SourceTrajectories[PreviewData->SourceIndex];
+
+			if (!SourceTrajectory.TraceInfo.IsEmpty())
+			{
+				const bool bEnabled = bParentEnabled;
+				const ESlateDrawEffect DrawEffects = bEnabled ? ESlateDrawEffect::None : ESlateDrawEffect::DisabledEffect;
+
+				TRange<double> LocalViewRange = Timeline->GetTimeRange();
+				const float LocalViewRangeMin = LocalViewRange.GetLowerBoundValue();
+				const float LocalViewRangeMax = LocalViewRange.GetUpperBoundValue();
+				const float LocalSequenceLength = LocalViewRangeMax-LocalViewRangeMin;
+				FSlateFontInfo SmallLayoutFont = FCoreStyle::GetDefaultFontStyle("Regular", 8);
+				
+				FVector2D Scale = FVector2D(1.0f,1.0f);
+				if ( LocalSequenceLength > 0)
+				{
+					SSimpleTimeSlider::FScrubRangeToScreen RangeToScreen( LocalViewRange, AllottedGeometry.GetLocalSize() );
+			
+					// Draw section
+					// TRange<double> ClampRange = Timeline->GetClampRange();
+
+					for (int i = 0; i < SourceTrajectory.TraceInfo.Num(); ++i)
+					{
+						int CurrentSplitIndex = SourceTrajectory.TraceInfo.Ranges[i].GetLowerBoundValue();
+						int NextSplitIndex = SourceTrajectory.TraceInfo.Ranges[i].GetUpperBoundValue();
+
+						const TArray<FTransform>& CurrentPoseRef = SourceTrajectory.Poses[CurrentSplitIndex];
+						const TArray<FTransform>& NextPoseRef = SourceTrajectory.Poses[NextSplitIndex];
+
+						const FGameplayTrajectory::FSample& CurrentSampleRef = SourceTrajectory.Samples[CurrentSplitIndex];
+						const FGameplayTrajectory::FSample& NextSampleRef = SourceTrajectory.Samples[NextSplitIndex];
+						
+						PaintBoneSection(AllottedGeometry, SmallLayoutFont, OutDrawElements, LayerId, RangeToScreen, CurrentPoseRef.Num(), CurrentSampleRef.Time, NextSampleRef.Time, i);
+					}
+				}	
+			}
+		
+		}
+		
+		return LayerId;
+	}
+	/** End SWidget */
+
+private:
+
+	/** Used to draw the trajectory */
+	TSharedPtr<FPreviewData> PreviewData;
+
+	/** Trajectory timeline widget */
+	TSharedPtr<SSimpleTimeSlider> Timeline;
+	
+	/** Delegate to invoke when the context menu should be opening. If it is nullptr, a context menu will not be summoned. */
+	FOnContextMenuOpening OnContextMenuOpening;
+
+	/** Delegate to invoke when the scrub position was changed. */
+	SSimpleTimeSlider::FOnScrubPositionChanged OnScrubPositionChanged;
+};
+
+void SExportTrajectoriesWindow::Construct(const FArguments& InArgs)
+{
+	WidgetWindow = InArgs._WidgetWindow;
+	Trajectories.Append(InArgs._Trajectories);
+	TrajectoryOwnerNames.Append(InArgs._OwnerNames);
+	PreviewData = MakeShared<FPreviewData>();
+
+	check(Trajectories.Num() == TrajectoryOwnerNames.Num());
+	
+	ImmutableState.SkeletalMeshInfos = InArgs._SkelMeshInfos;
+	ImmutableState.DebugInfos = InArgs._DebugInfos;
+	
+	check(Trajectories.Num() == ImmutableState.SkeletalMeshInfos.Num() && ImmutableState.SkeletalMeshInfos.Num() == ImmutableState.DebugInfos.Num())	
+	
+	// Export configuration settings.
+	FName Name = FName("TrajectoryExportDetails");
+	Name = MakeUniqueObjectName(GetTransientPackage(), UTrajectoryExportDetails::StaticClass(), Name, EUniqueObjectNameOptions::GloballyUnique);
+	ExportDetails = NewObject<UTrajectoryExportDetails>( GetTransientPackage(), Name, RF_Public | RF_Standalone);
+
+	// Allow user to view and edit the configuration settings.
+	{
+		FDetailsViewArgs GridDetailsViewArgs;
+		GridDetailsViewArgs.bAllowSearch = false;
+		GridDetailsViewArgs.NameAreaSettings = FDetailsViewArgs::HideNameArea;
+		GridDetailsViewArgs.bHideSelectionTip = true;
+		GridDetailsViewArgs.DefaultsOnlyVisibility = EEditDefaultsOnlyNodeVisibility::Automatic;
+		GridDetailsViewArgs.bShowOptions = false;
+		GridDetailsViewArgs.bAllowMultipleTopLevelObjects = false;
+
+		FPropertyEditorModule& PropertyEditorModule = FModuleManager::GetModuleChecked<FPropertyEditorModule>("PropertyEditor");
+		ExportDetailsView = PropertyEditorModule.CreateDetailView(GridDetailsViewArgs);
+		ExportDetailsView->SetObject(ExportDetails.Get());
+		ExportDetailsView->OnFinishedChangingProperties().AddSP(this, &SExportTrajectoriesWindow::OnFinishedChangingExportSettingsSelectionProperties);	
+	}
+	
+	ChildSlot
+	[
+		SNew(SHorizontalBox)
+		
+		+SHorizontalBox::Slot()
+		.HAlign(HAlign_Fill)
+		.VAlign(VAlign_Fill)
+		[
+			// Settings / Export
+			
+			SNew(SSplitter)
+			.Orientation(Orient_Horizontal)
+			/*.ResizeMode(ESplitterResizeMode::Fill)*/
+
+			// Viewport / Slider
+			
+			+ SSplitter::Slot()
+			.Value(0.6f)
+			[
+				SNew(SVerticalBox)
+				+ SVerticalBox::Slot()
+				.FillHeight(0.6)
+				[
+					SAssignNew(Viewport, SPreviewTrajectoryViewport)
+					.PreviewData(PreviewData)
+				]
+				
+				+ SVerticalBox::Slot()
+				.AutoHeight()
+				[
+					SAssignNew(ViewportTimeline, STrajectoryTimeline)
+					.PreviewData(PreviewData)
+					.OnScrubPositionChanged_Lambda([this](double NewSliderTimer, bool bIsScrubbing)
+					{
+						if (bIsScrubbing)
+						{
+							Viewport->ResetPreviewSkeletalMesh();
+						}
+					})
+					.OnContextMenuOpening_Lambda([InPreviewData = PreviewData, InExportWindow = this]()
+					{
+						// Early out
+						if (!InPreviewData.IsValid() || !IsValid(InPreviewData->Details))
+						{
+							return TSharedPtr<SWidget>();
+						}
+						
+						const int32 ScrubSampleIndex = GetSampleIndexFromTime(InPreviewData->Details->ExportSettings.FrameRate, InPreviewData->Timeline.ScrubTime, InPreviewData->OutputTrajectory.Samples.Num());
+						const double ScrubSampleTime = InPreviewData->OutputTrajectory.Samples[ScrubSampleIndex].Time;
+						
+						FMenuBuilder MenuBuilder(true, nullptr);
+
+						// Set min range via timeline
+						{
+							const FText Label = FText::Format(LOCTEXT("SetExportRangeMinLabel", "Set export range min at frame {0}* ({1}s)"), ScrubSampleIndex, ScrubSampleTime); 
+							const FText ToolTip = LOCTEXT("SetExportRangeMinToolTip", "Sets the ExportRange's Min value using the current scrub frame's time.");
+							FUIAction Action = FExecuteAction::CreateLambda([InPreviewData, InExportWindow]()
+							{
+								if (IsValid(InPreviewData->Details))
+								{
+									const int32 ScrubSampleIndex = InPreviewData->Details->ExportSettings.FrameRate.AsFrameTime(InPreviewData->Timeline.ScrubTime).FloorToFrame().Value;
+									InPreviewData->Details->ExportSettings.Range.Min = InPreviewData->OutputTrajectory.Samples[ScrubSampleIndex].Time;
+									
+									// @todo: Should this just be done in PostEditChange()?
+									InPreviewData->Details->PlayLength = InPreviewData->Details->ExportSettings.Range.Size();
+									InPreviewData->Details->NumberOfKeyFrames = InPreviewData->Details->ExportSettings.FrameRate.AsFrameTime(InPreviewData->Details->PlayLength).CeilToFrame().Value;
+									InPreviewData->Details->ExportSettings.OriginTime = FMath::Clamp(InPreviewData->Details->ExportSettings.OriginTime, InPreviewData->Details->ExportSettings.Range.Min, InPreviewData->Details->ExportSettings.Range.Max);
+
+									// Update asset info (since we care about which trace range we use)
+									InExportWindow->UpdateAssetInfo(InPreviewData->Details->ExportAssetInfo.Skeleton, InPreviewData->Details->ExportAssetInfo.SkeletalMesh);
+								}
+							});
+							
+							MenuBuilder.AddMenuEntry(Label, ToolTip, FSlateIcon(), Action);
+						}
+
+						// Set max range via timeline
+						{
+							const FText Label = FText::Format(LOCTEXT("SetExportRangeMaxLabel", "Set export range max at frame {0}* ({1}s)"), ScrubSampleIndex, ScrubSampleTime); 
+							const FText ToolTip = LOCTEXT("SetExportRangeMinToolTip", "Sets the ExportRange's Max value using the current scrub frame's time.");
+							FUIAction Action = FExecuteAction::CreateLambda([InPreviewData, InExportWindow]()
+							{
+								if (IsValid(InPreviewData->Details))
+								{
+									const int32 ScrubSampleIndex = InPreviewData->Details->ExportSettings.FrameRate.AsFrameTime(InPreviewData->Timeline.ScrubTime).FloorToFrame().Value;
+									InPreviewData->Details->ExportSettings.Range.Max = InPreviewData->OutputTrajectory.Samples[ScrubSampleIndex].Time;
+
+									// @todo: Should this just be done in PostEditChange()?
+									InPreviewData->Details->PlayLength = InPreviewData->Details->ExportSettings.Range.Size();
+									InPreviewData->Details->NumberOfKeyFrames = InPreviewData->Details->ExportSettings.FrameRate.AsFrameTime(InPreviewData->Details->PlayLength).CeilToFrame().Value;
+									InPreviewData->Details->ExportSettings.OriginTime = FMath::Clamp(InPreviewData->Details->ExportSettings.OriginTime, InPreviewData->Details->ExportSettings.Range.Min, InPreviewData->Details->ExportSettings.Range.Max);
+
+									// Update asset info (since we care about which trace range we use)
+									InExportWindow->UpdateAssetInfo(InPreviewData->Details->ExportAssetInfo.Skeleton, InPreviewData->Details->ExportAssetInfo.SkeletalMesh);
+								}
+							});
+							
+							MenuBuilder.AddMenuEntry(Label, ToolTip, FSlateIcon(), Action);
+						}
+
+						// Set origin via timeline
+						{
+							const FText Label = FText::Format(LOCTEXT("SetThisFrameAsTheOriginLabel", "Set origin at frame {0}* ({1}s)"), ScrubSampleIndex, ScrubSampleTime); 
+							const FText ToolTip = LOCTEXT("SetExportRangeMinToolTip", "Sets the current scrub frame at the origin in baked out asset.");
+							FUIAction Action = FExecuteAction::CreateLambda([InPreviewData]()
+							{
+								if (IsValid(InPreviewData->Details))
+								{
+									const int32 ScrubSampleIndex = InPreviewData->Details->ExportSettings.FrameRate.AsFrameTime(InPreviewData->Timeline.ScrubTime).FloorToFrame().Value;
+									const double ScrubSampleTime = InPreviewData->OutputTrajectory.Samples[ScrubSampleIndex].Time;
+
+									InPreviewData->Details->ExportSettings.bShouldForceOrigin = true;
+									InPreviewData->Details->ExportSettings.OriginTime = FMath::Clamp(ScrubSampleTime, InPreviewData->Details->ExportSettings.Range.Min, InPreviewData->Details->ExportSettings.Range.Max);
+								}
+							});
+							
+							MenuBuilder.AddMenuEntry(Label, ToolTip, FSlateIcon(), Action);
+						}
+						
+						return MenuBuilder.MakeWidget().ToSharedPtr();
+					})
+				]	
+			]
+			
+			+ SSplitter::Slot()
+			.SizeRule(SSplitter::SizeToContent)
+			.Value(0.4f)
+			[
+				SNew(SVerticalBox)
+
+				// Trajectory Picker Row
+				+ SVerticalBox::Slot()
+				.HAlign(HAlign_Fill)
+				.AutoHeight()
+				[
+					SNew(SBorder)
+					.HAlign(HAlign_Fill)
+					.BorderImage(FAppStyle::GetBrush("ToolPanel.GroupBorder"))
+					[
+						SNew(SBorder)
+						.Padding(2.0f)
+						.BorderImage(FAppStyle::GetBrush("Brushes.Header"))
+						[
+							SNew(SHorizontalBox)
+
+							// Trajectory to export label
+							+ SHorizontalBox::Slot()
+							.VAlign(VAlign_Center)
+							.Padding(6.0f, 3.0f)
+							.AutoWidth()
+							[
+								SNew(SRichTextBlock)
+								.Text(LOCTEXT("TrajectyPickerLabel", "Trajectory to export"))
+								.TransformPolicy(ETextTransformPolicy::ToUpper)
+								.DecoratorStyleSet(&FAppStyle::Get())
+								.TextStyle(FAppStyle::Get(), "DetailsView.CategoryTextStyle")
+							]
+
+							// Trajectory picker widget
+							+ SHorizontalBox::Slot()
+							.Padding(3.0f, 3.0f)
+							.AutoWidth()
+							[
+								SAssignNew(TrajectoryPickerComboButton, SComboButton)
+								.VAlign(VAlign_Center)
+								.OnGetMenuContent(this, &SExportTrajectoriesWindow::OnGetTrajectoryPickerMenuContent)
+								.OnMenuOpenChanged(this, &SExportTrajectoriesWindow::OnTrajectoryPickerMenuOpened)
+								.ButtonContent()
+								[
+									SNew(STextBlock)
+									.Justification(ETextJustify::Center)
+									.Text_Lambda([this](){ return SelectedTrajectoryIndex == INDEX_NONE ? FText::FromName(NAME_None): FText::FromName(TrajectoryOwnerNames[SelectedTrajectoryIndex]); })
+								]
+							]
+						]
+					]
+				]
+
+				// Settings
+				+ SVerticalBox::Slot()
+				.HAlign(HAlign_Fill)
+				.VAlign(VAlign_Fill)
+				.FillHeight(0.6f)
+				[
+					SNew(SBorder)
+					.Padding(0)
+					[
+						ExportDetailsView.ToSharedRef()
+					]
+				]
+
+				// Export
+				+ SVerticalBox::Slot()
+				.VAlign(VAlign_Bottom)
+				.FillHeight(0.2f)
+				.AutoHeight()
+				[
+					SNew(SBorder)
+					.Padding(6.0f)
+					.BorderImage(FAppStyle::GetBrush("ToolPanel.GroupBorder"))
+					[
+						SAssignNew(ExportButton, SButton)
+						.ButtonStyle(FAppStyle::Get(), "Button")
+						.TextStyle(FAppStyle::Get(), "DialogButtonText")
+						.ToolTipText(LOCTEXT("BakeTrajectoryButtonToolTip", "Bake out trajectory to the target asset"))
+						.HAlign(HAlign_Center)
+						.IsEnabled_Lambda([this]()
+						{
+							// Early out if export range contains multiple trace ranges.
+							if (PreviewData->OutputTrajectory.TraceInfo.Num() > 1)
+							{
+								int MinIndex = GetSampleIndexFromTime(PreviewData->Details->ExportSettings.FrameRate, PreviewData->Details->ExportSettings.Range.Min, PreviewData->OutputTrajectory.Samples.Num());
+								int MaxIndex = GetSampleIndexFromTime(PreviewData->Details->ExportSettings.FrameRate, PreviewData->Details->ExportSettings.Range.Max, PreviewData->OutputTrajectory.Samples.Num());
+
+								FTrajectoryToolsLibrary::FRangeOverlapTestResult TestResult{};
+								FTrajectoryToolsLibrary::GetRangeOverlaps(PreviewData->OutputTrajectory, TRange<int32>{MinIndex, MaxIndex}, TestResult);
+
+								if (TestResult.bOverlaps && TestResult.Ranges.Num() > 1)
+								{
+									return false;
+								}
+							}
+							
+							return SelectedTrajectoryIndex != INDEX_NONE && ExportDetails && ExportDetails->ExportAssetInfo.IsValid() && ExportDetails->ExportAssetInfo.CanCreateAsset();
+						})
+						.OnClicked_Lambda([this]()
+						{
+							// Export trajectories to assets.
+							UTrajectoryExportOperation::ExportTrajectory(&Trajectories[SelectedTrajectoryIndex], ExportDetails->ExportSettings, ExportDetails->ExportAssetInfo, TrajectoryOwnerNames[SelectedTrajectoryIndex].ToString());
+
+							// Close window after operation is done.
+							if (WidgetWindow.IsValid())
+							{
+								WidgetWindow.Pin()->RequestDestroyWindow();
+							}
+							
+							return FReply::Handled();
+						})
+						.Text(LOCTEXT("ExportButton", "Bake Out"))
+					]
+				]
+			]
+			
+		]
+	];
+}
+
+void SExportTrajectoriesWindow::AddReferencedObjects(FReferenceCollector& Collector)
+{
+	Collector.AddReferencedObject(ExportDetails);
+}
+
+TSharedRef<SWidget> SExportTrajectoriesWindow::OnGetTrajectoryPickerMenuContent()
+{
+	const bool bInShouldCloseWindowAfterMenuSelection = true;
+	FMenuBuilder MenuBuilder(bInShouldCloseWindowAfterMenuSelection, nullptr);
+
+	MenuBuilder.BeginSection(NAME_None, FText::GetEmpty());
+	
+	for (int i = 0; i < TrajectoryOwnerNames.Num(); ++i)
+	{
+		const bool bIsValidTrajectory = Trajectories[i].IsValid();
+		const FText ToolTip = bIsValidTrajectory ? FText::FromName(TrajectoryOwnerNames[i]) : LOCTEXT("InvalidTrajectoryTooltip", "Invalid. Trajectory data is empty or there is a mistmatch between poses and samples.");
+
+		FUIAction UIAction;
+
+		UIAction.CanExecuteAction = FCanExecuteAction::CreateLambda([i, this]()
+		{
+			return Trajectories[i].IsValid();
+		});
+		
+		UIAction.ExecuteAction = FExecuteAction::CreateLambda([i, this]()
+		{
+			const int PrevSelectedTrajectoryIndex = SelectedTrajectoryIndex;
+			SelectedTrajectoryIndex = i;
+
+			if (PrevSelectedTrajectoryIndex != SelectedTrajectoryIndex && SelectedTrajectoryIndex != INDEX_NONE)
+			{
+				// Reset export settings
+				ExportDetails->Reset();
+				PreviewData->Details = ExportDetails;
+
+				// Update which trajectory we are selecting based on the index
+				PreviewData->SourceIndex = SelectedTrajectoryIndex;
+				PreviewData->Viewport.SourceSkeletalMeshInfo = ImmutableState.SkeletalMeshInfos[SelectedTrajectoryIndex];
+				PreviewData->Viewport.SourceSkeletalMeshComponentId = ImmutableState.DebugInfos[SelectedTrajectoryIndex].OwnerId;
+
+				// Allow to view raw data.
+				PreviewData->SourceTrajectories = MakeConstArrayView(Trajectories);
+				
+				// Match export frame rate
+				PreviewData->OutputTrajectory.Samples.Reset();
+				PreviewData->OutputTrajectory.Poses.Reset();
+				PreviewData->OutputTrajectory.TraceInfo.Reset();
+				FTrajectoryToolsLibrary::TransformTrajectoryToMatchFrameRate(Trajectories[SelectedTrajectoryIndex], ExportDetails->ExportSettings.FrameRate, PreviewData->OutputTrajectory);
+
+				// Always start with entire trajectory range marked for available for export
+				if (!PreviewData->OutputTrajectory.Samples.IsEmpty())
+				{
+					ExportDetails->ExportSettings.Range.Min = PreviewData->OutputTrajectory.Samples[0].Time;
+					ExportDetails->ExportSettings.Range.Max = PreviewData->OutputTrajectory.Samples.Last().Time;
+
+					ExportDetails->NumberOfKeyFrames = PreviewData->OutputTrajectory.Samples.Num();
+					ExportDetails->PlayLength = ExportDetails->ExportSettings.Range.Size();
+				}
+
+				// Update asset info
+				UpdateAssetInfo(ExportDetails->ExportAssetInfo.Skeleton, ExportDetails->ExportAssetInfo.SkeletalMesh);
+				
+				// Focus on preview trajectory
+				PreviewViewportFocusOnTrajectory(PreviewData, Viewport->GetViewportClient());
+
+				// Reset viewport.
+				Viewport->ResetPreviewSkeletalMesh();
+			}
+		});
+		
+		MenuBuilder.AddMenuEntry(FText::FromName(TrajectoryOwnerNames[i]), ToolTip, FSlateIcon(), UIAction);
+	}
+	
+	MenuBuilder.EndSection();
+
+	return MenuBuilder.MakeWidget();
+}
+
+void SExportTrajectoriesWindow::OnTrajectoryPickerMenuOpened(bool bIsOpen)
+{
+	if (!bIsOpen)
+	{
+		TrajectoryPickerComboButton->SetMenuContent(SNullWidget::NullWidget);
+	}
+}
+
+void SExportTrajectoriesWindow::OnFinishedChangingExportSettingsSelectionProperties(const FPropertyChangedEvent& PropertyChangedEvent)
+{
+	if (!PropertyChangedEvent.Property ||  PropertyChangedEvent.ChangeType != EPropertyChangeType::ValueSet)
+	{
+		return;
+	}
+	
+	const bool ChangedTargetFrameRate = PropertyChangedEvent.Property->GetName() == "FrameRate";
+    if (ChangedTargetFrameRate)
+    {
+    	// Clear preview / output trajectory.
+    	PreviewData->OutputTrajectory.Samples.Reset();
+    	PreviewData->OutputTrajectory.Poses.Reset();
+    	
+    	// Trigger gameplay trajectory recompute.
+    	FTrajectoryToolsLibrary::TransformTrajectoryToMatchFrameRate(Trajectories[SelectedTrajectoryIndex], ExportDetails->ExportSettings.FrameRate, PreviewData->OutputTrajectory);
+    }
+
+    const bool ChangedExportRange = PropertyChangedEvent.Property->GetName() == "Min" || PropertyChangedEvent.Property->GetName() == "Max";
+    if (ChangedExportRange && IsValid(PreviewData->Details))
+    {
+    	// Range can't be bigger than play length
+    	PreviewData->Details->ExportSettings.Range.Min = FMath::Clamp(PreviewData->Details->ExportSettings.Range.Min, 0, PreviewData->OutputTrajectory.Samples.Last().Time);
+    	PreviewData->Details->ExportSettings.Range.Max = FMath::Clamp(PreviewData->Details->ExportSettings.Range.Max, 0, PreviewData->OutputTrajectory.Samples.Last().Time);
+
+    	// Origin can't be outside of range
+    	PreviewData->Details->ExportSettings.OriginTime = FMath::Clamp(PreviewData->Details->ExportSettings.OriginTime, PreviewData->Details->ExportSettings.Range.Min, PreviewData->Details->ExportSettings.Range.Max);
+
+    	// Update asset info (since we care about which trace range we use)
+    	UpdateAssetInfo(ExportDetails->ExportAssetInfo.Skeleton, ExportDetails->ExportAssetInfo.SkeletalMesh);
+    }
+}
+
+void SExportTrajectoriesWindow::UpdateAssetInfo(FSoftObjectPath&  OutSkeletonPath, FSoftObjectPath & OutSkeletalMeshPath) const
+{
+	if (!ImmutableState.SkeletalMeshInfos.IsValidIndex(SelectedTrajectoryIndex) || !IRewindDebugger::Instance() || !IRewindDebugger::Instance()->GetAnalysisSession())
+	{
+		OutSkeletonPath.Reset();
+		OutSkeletalMeshPath.Reset();
+		return;
+	}
+	
+	// Get data providers for traced objets.
+	IRewindDebugger* RewindDebugger = IRewindDebugger::Instance();
+	const TraceServices::IAnalysisSession* Session = RewindDebugger->GetAnalysisSession();
+	const IGameplayProvider* GameplayProvider = Session->ReadProvider<IGameplayProvider>("GameplayProvider");
+	const IAnimationProvider* AnimationProvider = Session->ReadProvider<IAnimationProvider>("AnimationProvider");
+	
+	if (GameplayProvider && AnimationProvider)
+	{
+		// Start reading this trace time information.
+		TraceServices::FAnalysisSessionReadScope AnalysisSessionReadScope(*Session);
+
+		// Preview using skeletal mesh from current trace range.
+		if (PreviewData->OutputTrajectory.TraceInfo.Num() > 1)
+		{
+			int MinIndex = GetSampleIndexFromTime(PreviewData->Details->ExportSettings.FrameRate, PreviewData->Details->ExportSettings.Range.Min, PreviewData->OutputTrajectory.Samples.Num());
+			int MaxIndex = GetSampleIndexFromTime(PreviewData->Details->ExportSettings.FrameRate, PreviewData->Details->ExportSettings.Range.Max, PreviewData->OutputTrajectory.Samples.Num());
+
+			FTrajectoryToolsLibrary::FRangeOverlapTestResult TestResult{};
+			FTrajectoryToolsLibrary::GetRangeOverlaps(PreviewData->OutputTrajectory, TRange<int32>{MinIndex, MaxIndex}, TestResult);
+
+			if (TestResult.bOverlaps && TestResult.Ranges.Num() == 1)
+			{
+				FSkeletalMeshInfo SkeletalMeshInfo = PreviewData->OutputTrajectory.TraceInfo.SkeletalMeshInfos[TestResult.Ranges.Last()];
+			
+				// Ensure our skeleton is valid.
+				if (SkeletalMeshInfo.SkeletonId != 0)
+				{
+					const FObjectInfo& SkeletonObjectInfo = GameplayProvider->GetObjectInfo(SkeletalMeshInfo.SkeletonId);
+					OutSkeletonPath = FSoftObjectPath(SkeletonObjectInfo.PathName);
+					
+					return;
+				}
+			}
+		}
+
+		// Preview using initial traced skeletal mesh.
+		if (const FSkeletalMeshInfo* SkeletalMeshInfo = AnimationProvider->FindSkeletalMeshInfo(ImmutableState.SkeletalMeshInfos[SelectedTrajectoryIndex].Id))
+		{
+			if (const FObjectInfo* SkeletalMeshObjectInfo = GameplayProvider->FindObjectInfo(ImmutableState.SkeletalMeshInfos[SelectedTrajectoryIndex].Id))
+			{
+				OutSkeletalMeshPath = FSoftObjectPath(SkeletalMeshObjectInfo->PathName);
+				
+				// Ensure our skeleton is valid.
+				if (SkeletalMeshInfo->SkeletonId != 0)
+				{
+					const FObjectInfo& SkeletonObjectInfo = GameplayProvider->GetObjectInfo(SkeletalMeshInfo->SkeletonId);
+					OutSkeletonPath = FSoftObjectPath(SkeletonObjectInfo.PathName);
+					
+					return;
+				}
+			}
+		}
+	}
+
+	OutSkeletonPath.Reset();
+	OutSkeletalMeshPath.Reset();
+}
+
+#undef LOCTEXT_NAMESPACE
diff --git a/Engine/Plugins/Experimental/Animation/TrajectoryTools/Source/TrajectoryTools/Private/TrajectoryExportOperation.cpp b/Engine/Plugins/Experimental/Animation/TrajectoryTools/Source/TrajectoryTools/Private/TrajectoryExportOperation.cpp
new file mode 100644
--- /dev/null
+++ b/Engine/Plugins/Experimental/Animation/TrajectoryTools/Source/TrajectoryTools/Private/TrajectoryExportOperation.cpp
@@ -0,0 +1,513 @@
+// Copyright Epic Games, Inc. All Rights Reserved.
+
+#include "TrajectoryExportOperation.h"
+
+#include "AssetToolsModule.h"
+#include "ContentBrowserModule.h"
+#include "IContentBrowserSingleton.h"
+#include "ObjectTools.h"
+#include "TrajectoryLibrary.h"
+#include "AssetRegistry/AssetRegistryModule.h"
+#include "Factories/AnimSequenceFactory.h"
+#include "Framework/Notifications/NotificationManager.h"
+#include "Widgets/Notifications/SNotificationList.h"
+#include "Animation/AnimSequence.h"
+#include "Animation/Skeleton.h"
+#include "Engine/SkeletalMesh.h"
+
+#define LOCTEXT_NAMESPACE "TrajectoryLibrary"
+
+// @todo: Need to use the #if EDITOR around code that needs the editor (transactions).
+
+FTrajectoryExportSettings::FTrajectoryExportSettings()
+{
+	Reset();
+}
+
+bool FTrajectoryExportSettings::IsValid() const
+{
+	const bool bValidFrameRate = FrameRate.IsValid();
+	const bool bValidRange = Range.IsValid() && Range.Size() != 0;
+	const bool bValidOriginTime = bShouldForceOrigin ? Range.Contains(OriginTime) : true;
+	
+	return bValidFrameRate && bValidRange && bValidOriginTime;
+}
+
+void FTrajectoryExportSettings::Reset()
+{
+	FrameRate = FFrameRate(30, 1);
+	Range = {0, 0};
+	bShouldForceOrigin = false;
+	OriginTime = 0;
+	bShouldOverwriteExistingFiles = false;
+	bShouldExportOnlyAnimatedBones = true;
+}
+
+bool FTrajectoryExportAssetInfo::CanCreateAsset() const
+{
+	FString ObjectPath = FolderPath.Path / AssetName;
+	const FString GameDir = TEXT("/Game/");
+	const FString EngineDir = TEXT("/Engine/");
+
+	bool bDirectoryExists = true;
+	
+	if (ObjectPath.StartsWith(GameDir))
+	{
+		FString ProjectDir = FPaths::ProjectDir();
+		const FString RelativeAssetPath = FPaths::Combine(FPaths::ProjectContentDir(), ObjectPath.Replace(*GameDir, TEXT("")));
+		bDirectoryExists = FPaths::DirectoryExists(FPaths::ConvertRelativePathToFull(RelativeAssetPath));
+	}
+	else if (ObjectPath.StartsWith(EngineDir))
+	{
+		FString ProjectDir = FPaths::ProjectDir();
+		const FString RelativeAssetPath = FPaths::Combine(FPaths::EngineContentDir(), ObjectPath.Replace(*EngineDir, TEXT("")));
+		bDirectoryExists = FPaths::DirectoryExists(FPaths::ConvertRelativePathToFull(RelativeAssetPath));
+	}
+	
+	if (bDirectoryExists)
+	{
+		return false;	
+	}
+	
+	IAssetRegistry& AssetRegistry = FModuleManager::LoadModuleChecked<FAssetRegistryModule>(TEXT("AssetRegistry")).Get();
+	const FAssetData AssetData = AssetRegistry.GetAssetByObjectPath(FSoftObjectPath(ObjectPath));
+	
+	return !AssetData.IsValid();
+}
+
+void FTrajectoryExportAssetInfo::Reset()
+{
+	AssetName = "";
+	FolderPath.Path = TEXT("/Game");
+	Skeleton.Reset();
+	SkeletalMesh.Reset();
+}
+
+bool FTrajectoryExportAssetInfo::IsValid() const
+{
+	return Skeleton.IsAsset() && SkeletalMesh.IsAsset() && !AssetName.IsEmpty() && !FolderPath.Path.IsEmpty();
+}
+
+void FTrajectoryExportContext::Reset()
+{
+	Settings.Reset();
+	AssetInfo.Reset();
+	SourceObjectName.Reset();
+	Data = nullptr;
+}
+
+bool FTrajectoryExportContext::IsValid() const
+{
+	return Settings.IsValid() && AssetInfo.IsValid() && Data != nullptr;
+}
+
+void UTrajectoryExportOperation::ExportTrajectory(FGameplayTrajectory* InTrajectory, const FTrajectoryExportSettings& InSettings, const FTrajectoryExportAssetInfo& InAssetInfo, const FString& InSourceObjectName)
+{
+	FTrajectoryExportContext Context;
+	Context.Data = InTrajectory;
+	Context.Settings = InSettings;
+	Context.AssetInfo = InAssetInfo;
+	Context.SourceObjectName = InSourceObjectName;
+	
+	// Actually run the batch operation
+	const TStrongObjectPtr BatchOperation(NewObject<UTrajectoryExportOperation>());
+	BatchOperation->Run(Context);
+}
+
+void UTrajectoryExportOperation::Run(const FTrajectoryExportContext& Context)
+{
+	Reset();
+
+	// @todo: Validate context and log errors.
+	
+	// Keep track of progress.
+	constexpr int32 NumProgressSteps = 6; // Gen Assets: 1, Export To Assets: 4, Notify User: 1.
+	FScopedSlowTask Progress(NumProgressSteps, LOCTEXT("ExportingOperationProgress", "Exporting assets..."));
+
+	// Show progress dialog.
+	constexpr bool bShowCancelButton = true;
+	Progress.MakeDialog(bShowCancelButton);
+
+	// Wrap all changes into a single commit while running the export operation.
+	ActiveTransaction = MakeUnique<FScopedTransaction>(TEXT("TrajectoryExportOperation"), LOCTEXT("TrajectoryExportOperationTransaction", "Exporting trajectories"), nullptr);
+
+	// Start exporting
+	GenerateAssets(Context, Progress);
+	ExportDataToAssets(Context, Progress);
+	NotifyUserOfResults(Context, Progress);
+	CleanupIfCancelled(Progress);
+
+	// Reset our open transaction to commit it now that the operation is completed.
+	ActiveTransaction.Reset();
+}
+
+void UTrajectoryExportOperation::Reset()
+{
+	AssetToProcess = nullptr;
+	GeneratedAsset = nullptr;
+	ActiveTransaction.Reset();
+}
+
+void UTrajectoryExportOperation::GenerateAssets(const FTrajectoryExportContext& Context, FScopedSlowTask& Progress)
+{
+	Progress.EnterProgressFrame(1.0f, LOCTEXT("ExportOperationProgress_GenerateAssets", "Generating asset(s)..."));
+	
+	FString FinalPackageName;
+	FString FinalAssetName;
+
+	FAssetToolsModule& AssetToolsModule = FModuleManager::LoadModuleChecked<FAssetToolsModule>("AssetTools");
+	
+	// Get unique names before creating the respective package and asset(s).
+	{
+		FString PackagePath = Context.AssetInfo.FolderPath.Path.IsEmpty() ? TEXT("/Game") : Context.AssetInfo.FolderPath.Path;
+		FString BasePackagePath = Context.AssetInfo.FolderPath.Path + "/"; //Context.AssetInfo.AssetName;
+		
+		AssetToolsModule.Get().CreateUniqueAssetName(BasePackagePath, Context.AssetInfo.AssetName, FinalPackageName,FinalAssetName);		
+	}
+
+	// Attempt to overwrite asset(s) if possible.
+	if (Context.Settings.bShouldOverwriteExistingFiles)
+	{
+		UPackage* ExistingPackage = FindPackage(nullptr, *FinalPackageName);
+		UAnimSequence* ExistingObject = Cast<UAnimSequence>(StaticFindObject(UAnimSequence::StaticClass(), ExistingPackage, *FinalAssetName));
+		
+		if (ExistingObject)
+		{
+			// store in batch operation variable
+			AssetToProcess = ExistingObject;
+			return;
+		}
+	}
+
+	if (FinalPackageName.EndsWith(FinalAssetName))
+	{
+		FinalPackageName = FinalPackageName.LeftChop(FinalAssetName.Len() + 1);
+	}
+
+	// Load traced skeleton data.
+	USkeleton* Skeleton = TSoftObjectPtr<USkeleton>(Context.AssetInfo.Skeleton).LoadSynchronous();
+	USkeletalMesh* SkeletalMesh =  TSoftObjectPtr<USkeletalMesh>(Context.AssetInfo.SkeletalMesh).LoadSynchronous();
+
+	const bool bIsSkeletalMeshValid = IsValid(SkeletalMesh);
+	const bool bIsSkeletonValid = IsValid(SkeletalMesh);
+	
+	if (!bIsSkeletonValid && bIsSkeletalMeshValid)
+	{
+		Skeleton = SkeletalMesh->GetSkeleton();
+	}
+	else if (!bIsSkeletalMeshValid && bIsSkeletonValid)
+	{
+		SkeletalMesh = Skeleton->GetPreviewMesh();
+	}
+	else if (!bIsSkeletalMeshValid && !bIsSkeletonValid)
+	{
+		checkNoEntry();
+	}
+	
+	// Create new asset for storing trajectory data.
+	UAnimSequenceFactory* Factory = NewObject<UAnimSequenceFactory>();
+	Factory->TargetSkeleton = Skeleton;
+	Factory->PreviewSkeletalMesh = SkeletalMesh;
+	UAnimSequence* NewAsset = Cast<UAnimSequence>(AssetToolsModule.Get().CreateAsset(*FinalAssetName, *FinalPackageName/* FPackageName::GetLongPackagePath(FinalPackageName)*/, UAnimSequence::StaticClass(), Factory));
+
+	check(IsValid(NewAsset))
+	
+	GeneratedAsset = NewAsset;
+	AssetToProcess = GeneratedAsset;
+
+	// Inform asset registry of our new asset.
+	FAssetRegistryModule::AssetCreated(NewAsset);
+}
+
+void UTrajectoryExportOperation::ExportDataToAssets(const FTrajectoryExportContext& Context, FScopedSlowTask& Progress) const
+{
+	// Abort invalid asset was marked for processing.
+	if (!AssetToProcess.IsValid() || !IsValid(AssetToProcess.Get()))
+	{
+		return;
+	}
+	
+	// A valid skeleton is needed to query the ref pose skeleton.
+	USkeleton* Skeleton = TSoftObjectPtr<USkeleton>(Context.AssetInfo.Skeleton).LoadSynchronous();
+	if (!IsValid(Skeleton) || Skeleton->HasAnyFlags(RF_BeginDestroyed | RF_FinishDestroyed))
+	{
+		return;
+	}
+
+	// No bones in ref skeleton. Can't convert traced poses from component to local space.
+	const FReferenceSkeleton& RefSkeleton = Skeleton->GetReferenceSkeleton();
+	if (RefSkeleton.GetRawBoneNum() == 0)
+	{
+		return;
+	}
+
+	// Inconsistent tracing of samples vs poses. Something wrong happened.
+	const FGameplayTrajectory& RawTrajectory = *Context.Data;
+	if (RawTrajectory.Samples.Num() != RawTrajectory.Poses.Num())
+	{
+		return;
+	}
+
+	// No data provided. Abort.
+	if (RawTrajectory.Samples.IsEmpty())
+	{
+		return;
+	}
+
+	// Transform trajectory data to match export settings.
+	FGameplayTrajectory FinalTrajectory = {};
+	{
+		if (Progress.ShouldCancel())
+		{
+			return;
+		}
+		
+		Progress.EnterProgressFrame(1.f, LOCTEXT("ExportOperationProgress_RawTrajectory", "Transforming raw trajectory to match specified export settings..."));
+		FTrajectoryToolsLibrary::TransformTrajectoryToMatchExportSettings(RawTrajectory, Context.Settings, FinalTrajectory);
+	}
+
+	// Properly configure newly animation sequence.
+	{
+		if (Progress.ShouldCancel())
+		{
+			return;
+		}
+		
+		Progress.EnterProgressFrame(1.f, FText::Format(LOCTEXT("ExportOperationProgress_PreparingAsset", "Preparing asset: {0}"), FText::FromString(*AssetToProcess->GetName())));
+
+		// @todo: store prev info (skeleton), so we can revert if overwriting asset.
+		AssetToProcess->Modify(true);
+
+		// Update skeleton info.
+		{
+			FScopedTransaction Transaction(LOCTEXT("UndoAction_ModifyAnimSequence", "Applying Skeleton to Animation Sequence(s)"));
+			
+			AssetToProcess->SetSkeleton(Skeleton);
+			AssetToProcess->SetPreviewMesh(Skeleton->GetPreviewMesh());
+			AssetToProcess->bEnableRootMotion = true;
+		}
+		
+		const bool bShouldTransact = GeneratedAsset == nullptr;
+		IAnimationDataController& Controller = AssetToProcess->GetController();
+		IAnimationDataController::FScopedBracket ScopedBracket(Controller, LOCTEXT("ExportTrajectoryConfigureAsset_Bracket", "Configure new animation sequence"), bShouldTransact);
+
+		// Clean previous info, if any.
+		Controller.InitializeModel();
+		AssetToProcess->ResetAnimation();
+
+		// Ensure anim sequence playback matches export settings.
+		Controller.SetFrameRate(Context.Settings.FrameRate);
+		Controller.SetNumberOfFrames(FinalTrajectory.Samples.Num() - 1);
+		
+		// Ensure UI displays proper framerate.
+		AssetToProcess->ImportFileFramerate = static_cast<float>(Context.Settings.FrameRate.AsDecimal());
+		AssetToProcess->ImportResampleFramerate = static_cast<int32>(Context.Settings.FrameRate.AsInterval());
+
+		// Add all bone tracks in ref skeleton for now. @todo: Take into consideration bShouldExportOnlyAnimatedBones.
+		const int32 NumBones = RefSkeleton.GetRawBoneNum();
+		const TArray<FTransform>& RefBonePose = RefSkeleton.GetRawRefBonePose();
+		for (int32 BoneIndex = 0; BoneIndex < NumBones; ++BoneIndex)
+		{
+			const FName BoneName = RefSkeleton.GetBoneName(BoneIndex);
+			Controller.AddBoneCurve(BoneName);
+			Controller.SetBoneTrackKeys(BoneName, { FVector3f(RefBonePose[BoneIndex].GetTranslation()), FVector3f(RefBonePose[BoneIndex].GetTranslation()) }, { FQuat4f(RefBonePose[BoneIndex].GetRotation()), FQuat4f(RefBonePose[BoneIndex].GetRotation()) }, { FVector3f(RefBonePose[BoneIndex].GetScale3D()), FVector3f(RefBonePose[BoneIndex].GetScale3D()) });
+		}
+		
+		Controller.NotifyPopulated();
+	}
+
+	// Buffer for final bone transforms to export.
+	const TArray<TArray<FTransform>> & ComponentSpacePoses = FinalTrajectory.Poses;
+	TArray<TArray<FTransform>> LocalSpacePoses = ComponentSpacePoses;
+	
+	// Convert all recorded poses to local space.
+	{
+		Progress.EnterProgressFrame(1.f, LOCTEXT("ExportOperationProgress_ConvertingToLocalSpace", "Converting trajectory data to local space..."));
+
+		const bool bShouldTransact = GeneratedAsset == nullptr;
+		IAnimationDataController& Controller = AssetToProcess->GetController();
+		Controller.OpenBracket(LOCTEXT("ExportTrajectoryConvertPosesToLocalSpace_Bracket", "Convert recorded poses to local space"), bShouldTransact);
+
+		// Ensure all poses are in local space before exporting.
+		for (int32 PoseIndex = 0; PoseIndex  < ComponentSpacePoses.Num(); ++PoseIndex)
+		{
+			if (Progress.ShouldCancel())
+			{
+				Controller.CloseBracket(bShouldTransact);
+				return;
+			}
+			
+			const int32 NumBones = ComponentSpacePoses[PoseIndex].Num();
+		
+			// Convert from component space to local space.
+			for (int32 BoneIndex = NumBones - 1; BoneIndex > 0; --BoneIndex)
+			{
+				const int32 ParentIndex = RefSkeleton.GetRawParentIndex(BoneIndex);
+				if (ParentIndex != INDEX_NONE)
+				{
+					LocalSpacePoses[PoseIndex][BoneIndex] = LocalSpacePoses[PoseIndex][BoneIndex].GetRelativeTransform(LocalSpacePoses[PoseIndex][ParentIndex]);
+				}
+			}
+		}
+
+		Controller.CloseBracket(bShouldTransact);
+	}
+
+	// Output trajectory information to animation sequence.
+	{
+		Progress.EnterProgressFrame(1.f, FText::Format(LOCTEXT("ExportOperationProgress_OutputtingDataToAsset ", "Outputting trajectory data to asset: {0}"), FText::FromString(*AssetToProcess->GetName())));
+
+		const bool bShouldTransact = GeneratedAsset == nullptr;
+		IAnimationDataController& Controller = AssetToProcess->GetController();
+		Controller.OpenBracket(LOCTEXT("ExportTrajectoryToAnimSequence_Bracket", "Export data to anim sequence"), bShouldTransact);
+		
+		// Root bone which to apply the trajectory transforms over time.
+		const FName RootBoneName = RefSkeleton.GetBoneName(0);
+
+		const int32 NumOfKeys = FinalTrajectory.Samples.Num();
+
+		// Query all bone track available in anim sequence.
+		TArray<FName> BoneTrackNames;
+		Controller.GetModel()->GetBoneTrackNames(BoneTrackNames);
+		
+		for (int32 BoneTrackIndex = 0; BoneTrackIndex < BoneTrackNames.Num(); ++BoneTrackIndex)
+		{
+			if (Progress.ShouldCancel())
+			{
+				Controller.CloseBracket(bShouldTransact);
+				return;
+			}
+			
+			const FName & TrackName = BoneTrackNames[BoneTrackIndex];
+			const bool bIsRootBone = RootBoneName == TrackName;
+
+			// Holds all keys to export per bone step.
+			FRawAnimSequenceTrack RawTrack;
+			RawTrack.PosKeys.SetNum(NumOfKeys);
+			RawTrack.RotKeys.SetNum(NumOfKeys);
+			RawTrack.ScaleKeys.SetNum(NumOfKeys);
+			
+			for (int32 KeyIndex = 0; KeyIndex < NumOfKeys; KeyIndex++)
+			{
+				if (bIsRootBone)
+				{
+					const FQuat& DoubleQuat = FinalTrajectory.Samples[KeyIndex].Orientation;
+					const FVector& DoubleVec = FinalTrajectory.Samples[KeyIndex].Position;
+					
+					RawTrack.PosKeys[KeyIndex] = FVector3f(DoubleVec.X, DoubleVec.Y, DoubleVec.Z);
+					RawTrack.RotKeys[KeyIndex] = FQuat4f(DoubleQuat.X, DoubleQuat.Y, DoubleQuat.Z, DoubleQuat.W);
+					RawTrack.ScaleKeys[KeyIndex] = FVector3f(1.0f, 1.0f, 1.0f);
+				}
+				else
+				{
+					const int32 TrackNameToBoneIndex = RefSkeleton.FindRawBoneIndex(TrackName);
+
+					if (TrackNameToBoneIndex != INDEX_NONE && LocalSpacePoses[KeyIndex].IsValidIndex(TrackNameToBoneIndex))
+					{
+						RawTrack.PosKeys[KeyIndex].X = LocalSpacePoses[KeyIndex][TrackNameToBoneIndex].GetTranslation().X;
+						RawTrack.PosKeys[KeyIndex].Y = LocalSpacePoses[KeyIndex][TrackNameToBoneIndex].GetTranslation().Y;
+						RawTrack.PosKeys[KeyIndex].Z = LocalSpacePoses[KeyIndex][TrackNameToBoneIndex].GetTranslation().Z;
+					
+						RawTrack.RotKeys[KeyIndex].W = LocalSpacePoses[KeyIndex][TrackNameToBoneIndex].GetRotation().W;
+						RawTrack.RotKeys[KeyIndex].X = LocalSpacePoses[KeyIndex][TrackNameToBoneIndex].GetRotation().X;
+						RawTrack.RotKeys[KeyIndex].Y = LocalSpacePoses[KeyIndex][TrackNameToBoneIndex].GetRotation().Y;
+						RawTrack.RotKeys[KeyIndex].Z = LocalSpacePoses[KeyIndex][TrackNameToBoneIndex].GetRotation().Z;
+					
+						RawTrack.ScaleKeys[KeyIndex].X = LocalSpacePoses[KeyIndex][TrackNameToBoneIndex].GetScale3D().X;
+						RawTrack.ScaleKeys[KeyIndex].Y = LocalSpacePoses[KeyIndex][TrackNameToBoneIndex].GetScale3D().Y;
+						RawTrack.ScaleKeys[KeyIndex].Z = LocalSpacePoses[KeyIndex][TrackNameToBoneIndex].GetScale3D().Z;	
+					}
+				}
+			}
+
+			// Output all traced keys for the current bone.
+			Controller.SetBoneTrackKeys(TrackName, RawTrack.PosKeys, RawTrack.RotKeys, RawTrack.ScaleKeys, bShouldTransact);
+		}
+		
+		// Inform now that we are done changing the anim sequence.
+		Controller.NotifyPopulated();
+
+		Controller.CloseBracket(bShouldTransact);
+	}
+}
+
+
+void UTrajectoryExportOperation::NotifyUserOfResults(const FTrajectoryExportContext& Context, FScopedSlowTask& Progress) const
+{
+	if (!AssetToProcess.IsValid() || !IsValid(AssetToProcess.Get()))
+	{
+		return;
+	}
+	
+	// Select all new assets and show them in the content browser.
+	const FContentBrowserModule& ContentBrowserModule = FModuleManager::Get().LoadModuleChecked<FContentBrowserModule>("ContentBrowser");
+	TArray<FAssetData> NewAssetsData = { AssetToProcess.Get() };
+	ContentBrowserModule.Get().SyncBrowserToAssets(NewAssetsData);
+
+	// Create pop-up notification in editor UI
+	constexpr float NotificationDuration = 5.0f;
+	if (Progress.ShouldCancel())
+	{
+		Progress.EnterProgressFrame(1.f, FText(LOCTEXT("ExportOperationProgress_Cancelled", "Cancelled.")));
+		
+		// Notify user that retarget was cancelled
+		FNotificationInfo Notification(FText::GetEmpty());
+		Notification.ExpireDuration = NotificationDuration;
+		Notification.Text = FText(LOCTEXT("CancelledExportOperation_NotificationTitle", "Export trajectory cancelled."));
+		FSlateNotificationManager::Get().AddNotification(Notification);
+	}
+	else
+	{
+		Progress.EnterProgressFrame(1.f, FText(LOCTEXT("ExportOperationProgress_Completed", "Export trajectory complete!")));
+		
+		// Log details of what assets were created
+		if (GeneratedAsset.IsValid())
+		{
+			UE_LOG(LogTemp, Display, TEXT("Export trajectory - New Asset Created: %s"), *GeneratedAsset->GetName());
+		}
+		else
+		{
+			UE_LOG(LogTemp, Display, TEXT("Export trajectory - Asset Modified: %s"), *AssetToProcess->GetName());
+		}
+		
+		// Notify user that export was completed
+		FNotificationInfo Notification(FText::GetEmpty());
+		Notification.ExpireDuration = NotificationDuration;
+		Notification.Text = FText::Format(
+			LOCTEXT("CompletedExportingTrajectoryNotification", "{0}'s trajectory data was exported to {1}. See Output for details."),
+			FText::FromString(Context.SourceObjectName),
+			FText::FromString(*AssetToProcess->GetName()));
+		FSlateNotificationManager::Get().AddNotification(Notification);
+	}
+}
+
+void UTrajectoryExportOperation::CleanupIfCancelled(const FScopedSlowTask& Progress) const
+{
+	if (!Progress.ShouldCancel())
+	{
+		return;
+	}
+
+	// Revert any changes.
+	if (ActiveTransaction.IsValid())
+	{
+		// We need to call Apply on the global undo, or cancelling the transaction doesn't actually roll back.
+		GUndo->Apply();
+		ActiveTransaction->Cancel();
+	}
+
+	// Any generated assets we just delete them since their changes were not transacted.
+	if (GeneratedAsset.IsValid())
+	{
+		// Notify the asset registry
+		FAssetRegistryModule::AssetDeleted(GeneratedAsset.Get());
+	
+		// Rename the asset we created out of the way
+		GeneratedAsset->Rename(nullptr, GetTransientPackage(), REN_DontCreateRedirectors | REN_NonTransactional);
+		
+		GeneratedAsset->MarkAsGarbage();
+		ObjectTools::DeleteAssets({ GeneratedAsset.Get() }, false);	
+	}
+}
+
+#undef LOCTEXT_NAMESPACE
diff --git a/Engine/Plugins/Experimental/Animation/TrajectoryTools/Source/TrajectoryTools/Private/TrajectoryLibrary.cpp b/Engine/Plugins/Experimental/Animation/TrajectoryTools/Source/TrajectoryTools/Private/TrajectoryLibrary.cpp
new file mode 100644
--- /dev/null
+++ b/Engine/Plugins/Experimental/Animation/TrajectoryTools/Source/TrajectoryTools/Private/TrajectoryLibrary.cpp
@@ -0,0 +1,314 @@
+﻿// Copyright Epic Games, Inc. All Rights Reserved.
+
+#include "TrajectoryLibrary.h"
+
+#include "Animation/AnimSequence.h"
+#include "UObject/Object.h"
+#include "Algo/BinarySearch.h"
+#include "BoneContainer.h"
+#include "BonePose.h"
+#include "TrajectoryExportOperation.h"
+
+#define LOCTEXT_NAMESPACE "TrajectoryLibrary"
+
+void FTrajectoryToolsLibrary::GetRangeOverlaps(const FGameplayTrajectory& InTrajectory, const TRange<int32>& InSampleRange, FRangeOverlapTestResult& InOutOverlapResult)
+{
+	int OverlapCount = 0;
+	
+	if (!InTrajectory.TraceInfo.IsEmpty() && !InTrajectory.Samples.IsEmpty() && !InSampleRange.IsEmpty() && !InSampleRange.IsDegenerate())
+	{
+		for (int i = 0; i < InTrajectory.TraceInfo.Num(); ++i)
+		{
+			if (bool bOverlaps =  InTrajectory.TraceInfo.Ranges[i].Overlaps(InSampleRange))
+			{
+				InOutOverlapResult.bOverlaps = bOverlaps;
+				InOutOverlapResult.Ranges.Push(i);
+				++OverlapCount;
+			}
+		}
+	}
+
+	ensureMsgf(OverlapCount <= 2, TEXT("This function should only be called where we know only two range are possibly overlapping"));
+}
+
+void FTrajectoryToolsLibrary::GetSampleIndicesForMatchedSampleTime(const FGameplayTrajectory& InTrajectory, double InRequestedSampleTime, TPair<int32, int32>& OutMatchedSampleIndicesRange)
+{
+	if (!InTrajectory.Samples.IsEmpty())
+	{
+		InRequestedSampleTime = FMath::Clamp(InRequestedSampleTime, 0, InTrajectory.Samples.Last().Time);
+		
+		const int32 EndSampleIndex = Algo::LowerBoundBy(InTrajectory.Samples, InRequestedSampleTime, &FGameplayTrajectory::FSample::Time);
+		check(EndSampleIndex < InTrajectory.Samples.Num())
+		const int32 StartSampleIndex = FMath::Clamp(EndSampleIndex - 1, 0, InTrajectory.Samples.Num() - 1);
+		
+		OutMatchedSampleIndicesRange.Get<0>() = StartSampleIndex;
+		OutMatchedSampleIndicesRange.Get<1>() = EndSampleIndex;
+	}
+	else
+	{
+		OutMatchedSampleIndicesRange.Get<0>() = INDEX_NONE;
+		OutMatchedSampleIndicesRange.Get<1>() = INDEX_NONE;
+	}
+}
+
+bool FTrajectoryToolsLibrary::GetTransformAtTimeInTrajectory(const FGameplayTrajectory& InTrajectory, double InRequestedTime, FTransform& OuTransform)
+{
+	// Get keys for smallest range that contains the requested time.
+	TPair<int32, int32> IndicesRangeThatMatchRequestedTime = { INDEX_NONE, INDEX_NONE };
+	GetSampleIndicesForMatchedSampleTime(InTrajectory, InRequestedTime, IndicesRangeThatMatchRequestedTime);
+
+	// Invalid indices.
+	if (IndicesRangeThatMatchRequestedTime.Get<0>() == INDEX_NONE || IndicesRangeThatMatchRequestedTime.Get<1>() == INDEX_NONE)
+	{
+		return false;
+	}
+
+	// Time requested exactly matches a sample (The case of the calling GetSampleIndicesForMatchedSampleTime() with time 0 or having trajectory with one sample).
+	if (IndicesRangeThatMatchRequestedTime.Get<0>() == IndicesRangeThatMatchRequestedTime.Get<1>())
+	{
+		const FGameplayTrajectory::FSample& Sample = InTrajectory.Samples[IndicesRangeThatMatchRequestedTime.Get<0>()];
+
+		OuTransform = FTransform(Sample.Orientation, Sample.Position);
+	}
+	// Interpolate keys to exactly match the requested time.
+	else
+	{
+		const FGameplayTrajectory::FSample& StartSample = InTrajectory.Samples[IndicesRangeThatMatchRequestedTime.Get<0>()];
+		const FGameplayTrajectory::FSample& EndSample = InTrajectory.Samples[IndicesRangeThatMatchRequestedTime.Get<1>()];
+		
+		const double Alpha = (InRequestedTime - StartSample.Time) / (EndSample.Time - StartSample.Time);
+
+		OuTransform = FTransform(FQuat::Slerp(StartSample.Orientation, EndSample.Orientation, Alpha), FMath::Lerp(StartSample.Position, EndSample.Position, Alpha));
+	}
+
+	return true;
+}
+
+bool FTrajectoryToolsLibrary::GetPoseAtTimeInTrajectory(const FGameplayTrajectory& InTrajectory, double InRequestedTime, TArray<FTransform>& OutCSPose)
+{
+	// Get keys for smallest range that contains the requested time.
+	TPair<int32, int32> IndicesRangeThatMatchRequestedTime = { INDEX_NONE, INDEX_NONE };
+	GetSampleIndicesForMatchedSampleTime(InTrajectory, InRequestedTime, IndicesRangeThatMatchRequestedTime);
+
+	// Invalid indices.
+	if (IndicesRangeThatMatchRequestedTime.Get<0>() == INDEX_NONE || IndicesRangeThatMatchRequestedTime.Get<1>() == INDEX_NONE)
+	{
+		return false;
+	}
+	
+	// OutCSPose.Reserve(InTrajectory.Poses[IndicesRangeThatMatchRequestedTime.Get<0>()].Num());
+	OutCSPose.SetNumZeroed(InTrajectory.Poses[IndicesRangeThatMatchRequestedTime.Get<0>()].Num());
+	
+	// Time requested exactly matches a sample (The case of the calling GetSampleIndicesForMatchedSampleTime() with time 0 or having trajectory with one sample).
+	if (IndicesRangeThatMatchRequestedTime.Get<0>() == IndicesRangeThatMatchRequestedTime.Get<1>())
+	{
+		const TArray<FTransform>& Pose = InTrajectory.Poses[IndicesRangeThatMatchRequestedTime.Get<0>()];
+
+		for (int i = 0; i < Pose.Num(); ++i)
+		{
+			OutCSPose[i] = Pose[i];
+		}
+	}
+	// Interpolate keys to exactly match the requested time.
+	else
+	{
+		FRangeOverlapTestResult TestResult{};
+		GetRangeOverlaps(InTrajectory, TRange<int32>(IndicesRangeThatMatchRequestedTime.Key, IndicesRangeThatMatchRequestedTime.Value), TestResult);
+
+		if (TestResult.bOverlaps && TestResult.Ranges.Num() == 2)
+		{
+			const FGameplayTrajectory::FSample& StartSample = InTrajectory.Samples[IndicesRangeThatMatchRequestedTime.Get<0>()];
+			const FGameplayTrajectory::FSample& EndSample = InTrajectory.Samples[IndicesRangeThatMatchRequestedTime.Get<1>()];
+
+			const TArray<FTransform>& StartPose = InTrajectory.Poses[IndicesRangeThatMatchRequestedTime.Get<0>()];
+			const TArray<FTransform>& EndPose = InTrajectory.Poses[IndicesRangeThatMatchRequestedTime.Get<1>()];
+
+			if (FMath::Abs(StartSample.Time - InRequestedTime) < FMath::Abs(EndSample.Time - InRequestedTime))
+			{
+				OutCSPose = StartPose;
+			}
+			else
+			{
+				OutCSPose = EndPose;
+			}
+			
+			return true;
+		}
+		
+		const TArray<FTransform>& StartPose = InTrajectory.Poses[IndicesRangeThatMatchRequestedTime.Get<0>()];
+		const TArray<FTransform>& EndPose = InTrajectory.Poses[IndicesRangeThatMatchRequestedTime.Get<1>()];
+
+		// check(StartPose.Num() == EndPose.Num() && EndPose.Num() == OutCSPose.Num());
+		
+		const FGameplayTrajectory::FSample& StartSample = InTrajectory.Samples[IndicesRangeThatMatchRequestedTime.Get<0>()];
+		const FGameplayTrajectory::FSample& EndSample = InTrajectory.Samples[IndicesRangeThatMatchRequestedTime.Get<1>()];
+		
+		const double Alpha = (InRequestedTime - StartSample.Time) / (EndSample.Time - StartSample.Time);
+
+		for (int i = 0; i < StartPose.Num(); ++i)
+		{
+			OutCSPose[i].Blend(StartPose[i], EndPose[i], Alpha);
+		}
+	}
+
+	return true;
+}
+
+void FTrajectoryToolsLibrary::TransformTrajectoryToMatchFrameRate(const FGameplayTrajectory& InTrajectory, FFrameRate InFrameRate, FGameplayTrajectory& OutFrameMatchedTrajectory)
+{
+	check(InTrajectory.Samples.Num() == InTrajectory.Poses.Num());
+	
+	if (!InTrajectory.Samples.IsEmpty())
+	{
+		if (InTrajectory.Samples.Num() == 1)
+		{
+			// Fast-path when only one sample is present.
+			OutFrameMatchedTrajectory.Samples = InTrajectory.Samples;
+			OutFrameMatchedTrajectory.Poses = InTrajectory.Poses;
+			OutFrameMatchedTrajectory.TraceInfo = InTrajectory.TraceInfo;
+		}
+		else
+		{
+			const double TrajectoryPlayLength = InTrajectory.Samples.Num() == 1 ? InTrajectory.Samples[0].Time : InTrajectory.Samples.Last().Time - InTrajectory.Samples[0].Time;
+			const double TrajectorySampleRate = InFrameRate.AsInterval();
+			const int32 TotalSamples = InFrameRate.AsFrameTime(TrajectoryPlayLength).CeilToFrame().Value;
+
+			checkf(TotalSamples > 0, TEXT("Total output of samples is zero or less."))
+
+			// Clear buffer before processing
+			OutFrameMatchedTrajectory.Samples.Reset(TotalSamples);
+			OutFrameMatchedTrajectory.Poses.Reset(TotalSamples);
+			OutFrameMatchedTrajectory.TraceInfo.Reset(InTrajectory.TraceInfo.Num());
+
+			// Trace ranges to match that of the source trajectory.
+			for (int32 i = 0; i < InTrajectory.TraceInfo.Num(); ++i)
+			{
+				const int LowerBound = InTrajectory.TraceInfo.Ranges[i].GetLowerBoundValue();
+				const int UpperBound = InTrajectory.TraceInfo.Ranges[i].GetUpperBoundValue();
+
+				const double LowerBoundSampleTime = InTrajectory.Samples[LowerBound].Time;
+				const double UpperBoundSampleTime = InTrajectory.Samples[UpperBound].Time;
+			
+				const int FrameMatchedLowerBoundIndex = FMath::Clamp(InFrameRate.AsFrameTime(LowerBoundSampleTime).FloorToFrame().Value, 0, TotalSamples - 1);
+				const int FrameMatchedUpperBoundIndex = FMath::Clamp(InFrameRate.AsFrameTime(UpperBoundSampleTime).FloorToFrame().Value, 0, TotalSamples - 1);
+
+				OutFrameMatchedTrajectory.TraceInfo.Ranges.Emplace(TRangeBound<int32>::Inclusive(FrameMatchedLowerBoundIndex), TRangeBound<int32>::Inclusive(FrameMatchedUpperBoundIndex));
+				OutFrameMatchedTrajectory.TraceInfo.SkeletalMeshInfos.Push(InTrajectory.TraceInfo.SkeletalMeshInfos[i]);
+			}
+		
+			for (int32 OutSampleIndex = 0; OutSampleIndex < TotalSamples; ++OutSampleIndex)
+			{
+				const double RequestedSampleTime = FMath::Clamp(static_cast<double>(OutSampleIndex) * TrajectorySampleRate, 0.0f, TrajectoryPlayLength);
+
+				FTransform SampleTransform;
+				check(GetTransformAtTimeInTrajectory(InTrajectory, RequestedSampleTime, SampleTransform))
+
+				OutFrameMatchedTrajectory.Samples.AddDefaulted();
+				FGameplayTrajectory::FSample& TimeMatchedSample = OutFrameMatchedTrajectory.Samples.Last();
+	
+				TimeMatchedSample.Position = SampleTransform.GetLocation();
+				TimeMatchedSample.Orientation = SampleTransform.GetRotation();
+				TimeMatchedSample.Time = RequestedSampleTime;
+
+				OutFrameMatchedTrajectory.Poses.AddDefaulted();
+				TArray<FTransform>& TimeMatchedPose = OutFrameMatchedTrajectory.Poses.Last();
+
+				check(GetPoseAtTimeInTrajectory(InTrajectory, RequestedSampleTime, TimeMatchedPose));
+			}
+		}
+	}
+	else
+	{
+		// @todo: Warning.
+	}
+}
+
+void FTrajectoryToolsLibrary::TransformTrajectoryToMatchExportSettings(const FGameplayTrajectory& InTrajectory, const FTrajectoryExportSettings& InExportSettings, FGameplayTrajectory& OutTransformedTrajectory)
+{
+	if (InTrajectory.Samples.IsEmpty())
+	{
+		// error trajectory empty
+		return;
+	}
+
+	if (!InExportSettings.IsValid())
+	{
+		// invalid export settings
+		return;
+	}
+
+	// Ensure samples match framerate
+	TransformTrajectoryToMatchFrameRate(InTrajectory, InExportSettings.FrameRate, OutTransformedTrajectory);
+	
+	// Make all samples be relative to the position at sample given by the origin time.
+	if (InExportSettings.bShouldForceOrigin)
+	{
+		const int SampleIndexAtOriginTime = InExportSettings.FrameRate.AsFrameTime(InExportSettings.OriginTime).FrameNumber.Value;
+		const FVector OffsetToOrigin = -OutTransformedTrajectory.Samples[SampleIndexAtOriginTime].Position;
+		
+		for (FGameplayTrajectory::FSample& Sample : OutTransformedTrajectory.Samples)
+		{
+			Sample.Position += OffsetToOrigin;
+		}
+	}
+	
+	// Prune samples not within range.
+	{
+		int StartRangeSampleIndex = INDEX_NONE;
+		int EndRangeSampleIndex = INDEX_NONE;
+		
+		// @todo: Improve with two pointer technique.
+		for (int SampleIndex = 0; SampleIndex <  OutTransformedTrajectory.Samples.Num(); ++SampleIndex)
+		{
+			FGameplayTrajectory::FSample& Sample = OutTransformedTrajectory.Samples[SampleIndex];
+
+			// Find lower bound index
+			if (StartRangeSampleIndex == INDEX_NONE)
+			{
+				const bool bPassesLowerBoundCheck = FMath::IsNearlyEqual(Sample.Time, InExportSettings.Range.Min) || Sample.Time > InExportSettings.Range.Min;
+
+				if (bPassesLowerBoundCheck)
+				{
+					StartRangeSampleIndex = SampleIndex;
+				}
+			}
+
+			// Find upper bound index
+			if (EndRangeSampleIndex == INDEX_NONE)
+			{
+				const bool bPassesUpperBoundCheck = FMath::IsNearlyEqual(Sample.Time, InExportSettings.Range.Max) || Sample.Time > InExportSettings.Range.Max;
+
+				if (bPassesUpperBoundCheck)
+				{
+					EndRangeSampleIndex = SampleIndex;
+				}
+			}
+		}
+
+		// We never found the end, but we have a start. Default end to last sample.
+		if (StartRangeSampleIndex != INDEX_NONE && EndRangeSampleIndex == INDEX_NONE)
+		{
+			EndRangeSampleIndex = OutTransformedTrajectory.Samples.Num() - 1;	
+		}
+		
+		// Update time
+		for (FGameplayTrajectory::FSample& Sample : OutTransformedTrajectory.Samples)
+		{
+			Sample.Time -= InExportSettings.Range.Min;
+		}
+
+		// @todo: Improve the efficiency of this by not copying so much.
+		
+		// Output pruned trajectory
+		const int TotalSamplesAfterPrune = EndRangeSampleIndex - StartRangeSampleIndex + 1;
+		const TArray<FGameplayTrajectory::FSample> PrunedTrajectory = TArray<FGameplayTrajectory::FSample>(MakeArrayView(OutTransformedTrajectory.Samples.GetData() + StartRangeSampleIndex, TotalSamplesAfterPrune));
+		const TArray<TArray<FTransform>> PrunedPoses = TArray<TArray<FTransform>>(MakeArrayView(OutTransformedTrajectory.Poses.GetData() + StartRangeSampleIndex, TotalSamplesAfterPrune));
+		
+		OutTransformedTrajectory.Samples.Reset();
+		OutTransformedTrajectory.Samples.Append(PrunedTrajectory);
+		OutTransformedTrajectory.Poses.Reset();
+		OutTransformedTrajectory.Poses.Append(PrunedPoses);
+	}
+}
+
+#undef LOCTEXT_NAMESPACE
diff --git a/Engine/Plugins/Experimental/Animation/TrajectoryTools/Source/TrajectoryTools/Private/TrajectoryRewindDebuggerExtension.cpp b/Engine/Plugins/Experimental/Animation/TrajectoryTools/Source/TrajectoryTools/Private/TrajectoryRewindDebuggerExtension.cpp
new file mode 100644
--- /dev/null
+++ b/Engine/Plugins/Experimental/Animation/TrajectoryTools/Source/TrajectoryTools/Private/TrajectoryRewindDebuggerExtension.cpp
@@ -0,0 +1,610 @@
+// Copyright Epic Games, Inc. All Rights Reserved.
+
+#include "TrajectoryRewindDebuggerExtension.h"
+
+#include "ToolMenus.h"
+#include "IAnimationProvider.h"
+#include "IGameplayProvider.h"
+#include "Debug/DebugDrawService.h"
+#include "DrawDebugHelpers.h"
+#include "Components/SkeletalMeshComponent.h"
+#include "Engine/Canvas.h"
+#include "SceneView.h"
+#include "SExportTrajectoriesWindow.h"
+#include "Engine/Engine.h"
+#include "Framework/MultiBox/MultiBoxBuilder.h"
+#include "Framework/Application/SlateApplication.h"
+#include "Widgets/SWindow.h"
+#include "Animation/AnimSequence.h"
+#include "Widgets/Input/SButton.h"
+#include "Logging/StructuredLog.h"
+
+#define LOCTEXT_NAMESPACE "RewindDebuggerTrajectory"
+
+FRewindDebuggerTrajectory::FRewindDebuggerTrajectory()
+{
+}
+
+void FRewindDebuggerTrajectory::Initialize()
+{
+	// Generate widgets.
+	GenerateMenu();
+
+	// Colors for trajectories. (Up to 4 rn, should cycle if trying to index out of bounds).
+	DebugDrawColors.Push(FColor::Cyan);
+	DebugDrawColors.Push(FColor::Emerald);
+	DebugDrawColors.Push(FColor::Magenta);
+	DebugDrawColors.Push(FColor::Yellow);
+}
+
+void FRewindDebuggerTrajectory::Shutdown()
+{
+	// Delete preview window.
+	if (BakeOutWindow.IsValid())
+	{
+		FSlateApplication::Get().DestroyWindowImmediately(BakeOutWindow.ToSharedRef());
+		BakeOutWindow = nullptr;
+	}
+}
+
+void FRewindDebuggerTrajectory::Reset()
+{
+	// Delete preview window. @todo: This is temporary.
+	if (BakeOutWindow.IsValid())
+	{
+		FSlateApplication::Get().DestroyWindowImmediately(BakeOutWindow.ToSharedRef());
+		BakeOutWindow = nullptr;
+	}
+	
+	// No valid world to visualize.
+	WorldToVisualize = nullptr;
+
+	// Stop hooking into debug draw.
+	if (DebugDrawDelegateHandle.IsValid())
+	{
+		UDebugDrawService::Unregister(DebugDrawDelegateHandle);
+		DebugDrawDelegateHandle.Reset();
+	}
+
+	// Clear state.
+	State.Reset();
+}
+
+void FRewindDebuggerTrajectory::Update(float DeltaTime, IRewindDebugger* RewindDebugger)
+{
+	check(RewindDebugger);
+	
+	// Early out.
+	if (RewindDebugger->IsPIESimulating() || RewindDebugger->GetRecordingDuration() == 0)
+	{
+		Reset();
+
+		if (RewindDebugger->IsPIESimulating())
+		{
+			State.bWasPIESimulating = true;
+		}
+		
+		return;
+	}
+
+	// @todo: Should we trace actor transforms?
+	// @todo: Use orientation from trajectory vs captured pose.
+
+	// Ensure we use RewindDebugger's world for debug drawing.
+	WorldToVisualize = RewindDebugger->GetWorldToVisualize();
+	
+	// Hook into engine flag for debug drawing.
+	EnsureDebugDrawDelegateExists();
+
+	// Only process trace file after PIE was paused.
+	if (State.bWasPIESimulating)
+	{
+		// Query info from trace file.
+		const TraceServices::IAnalysisSession* Session = RewindDebugger->GetAnalysisSession();
+		if (Session)
+		{
+			State.Extract.TraceStarTime = 0;
+			State.Extract.TraceEndTime = RewindDebugger->CurrentTraceTime();
+		
+			const IGameplayProvider* GameplayProvider = Session->ReadProvider<IGameplayProvider>("GameplayProvider");
+			const IAnimationProvider* AnimationProvider = Session->ReadProvider<IAnimationProvider>("AnimationProvider");
+
+			if (GameplayProvider && AnimationProvider)
+			{
+				TraceServices::FAnalysisSessionReadScope AnalysisSessionReadScope(*Session);
+			
+				BuildTrajectoryOwnersList(RewindDebugger, GameplayProvider, State.Extract.TraceStarTime, State.Extract.TraceEndTime, State.ObjectInfos);
+				BuildTrajectorySkeletalMeshInfoList(GameplayProvider, AnimationProvider, State.Extract.TraceStarTime, State.Extract.TraceEndTime, State.ObjectInfos, State.SkelMeshInfos);
+				BuildTrajectories(GameplayProvider, AnimationProvider, State.Extract.TraceStarTime, State.Extract.TraceEndTime, State.ObjectInfos, State.Trajectories);
+				UpdateDebugInfos(GameplayProvider, State.ObjectInfos, State.DebugInfos);
+			}
+		}
+		else
+		{
+			// Invalid trace session. Abort.
+			Reset();
+			return;
+		}
+		
+		State.bWasPIESimulating = false;
+	}
+	
+	check(State.Trajectories.Num() == State.ObjectInfos.Num());
+}
+
+void FRewindDebuggerTrajectory::RecordingStarted(IRewindDebugger* RewindDebugger)
+{
+	// Clear up any cached state / variables.
+	Reset();
+}
+
+void FRewindDebuggerTrajectory::RecordingStopped(IRewindDebugger* RewindDebugger)
+{
+	
+}
+
+void FRewindDebuggerTrajectory::Clear(IRewindDebugger* RewindDebugger)
+{
+	// Clear up any cached state / variables.
+	Reset();
+}
+
+void FRewindDebuggerTrajectory::GenerateMenu()
+{
+	if (UToolMenu* Menu = UToolMenus::Get()->FindMenu("RewindDebugger.ToolBar"))
+	{
+		FToolMenuSection& Section = Menu->AddSection("TrajectoryWorkflows", LOCTEXT("Trajectory Tools","Trajectory Tools"));
+		
+		Section.AddSeparator("TrajectoryWorkflows");
+		{
+			static const FName Name = "Trajectories";
+			static const FText Label = LOCTEXT("ToggleDebugDrawTrajectoriesLabel", "Trajectories");
+			static const FText ToolTip = LOCTEXT("ToggleDebugDrawTrajectoriesTooltip", "Toggle trajectories to display in viewport");
+			
+			Section.AddSubMenu(Name, Label, ToolTip, FNewToolMenuDelegate::CreateLambda([this](UToolMenu* InMenu)
+			{
+				FToolMenuSection& Section = InMenu->AddSection(NAME_None);
+				Section.AddDynamicEntry("Trajectories", FNewToolMenuSectionDelegate::CreateRaw(this, &FRewindDebuggerTrajectory::MakeTrajectoriesMenu));
+			}));
+		}
+	}
+}
+
+void FRewindDebuggerTrajectory::MakeTrajectoriesMenu(FToolMenuSection& InSection)
+{
+	// Toggle trajectory drawing
+	{
+		static const FName Name = "ToggleDebugDrawComboButton";
+		static const FText Label = LOCTEXT("ToggleDebugDrawTrajectoriesLabel", "Toggle debug draw");
+		static const FText ToolTip = LOCTEXT("ToggleDebugDrawTrajectoriesTooltip", "Toggle trajectories to display in viewport");
+
+		InSection.AddSubMenu(Name, Label, ToolTip, FNewToolMenuDelegate::CreateRaw(this, &FRewindDebuggerTrajectory::MakeToggleDebugDrawMenu), false, FSlateIcon(), false);
+	}
+
+	// Bake out trajectories
+	{
+		FUIAction UIAction;
+		UIAction.ExecuteAction = FExecuteAction::CreateLambda([this](){ ShowBakeOutTrajectoryWindow(); });
+		UIAction.CanExecuteAction = FCanExecuteAction::CreateLambda([this]()
+		{
+			IRewindDebugger* RewindDebugger = IRewindDebugger::Instance();
+			return RewindDebugger ? !RewindDebugger->IsPIESimulating() : false;
+		});
+		
+		static const FName Name = "BakeOutSubMenu"; 
+		static const FText Label = LOCTEXT("BakeOutTrajectoriesLabel", "Bake out...");
+		static const FText Tooltip = LOCTEXT("BakeOutTrajectoriesTooltip", "Bake a trajectory into a standalone asset");
+		
+		InSection.AddMenuEntry(Name, Label, Tooltip, FSlateIcon(), UIAction);
+	}
+}
+
+void FRewindDebuggerTrajectory::MakeToggleDebugDrawMenu(UToolMenu* InMenu)
+{
+	FToolMenuSection& Section = InMenu->FindOrAddSection("ToggleDebugDrawOptions");
+
+	for (FExtensionState::FDebugInfo& DebugInfoEntry : State.DebugInfos)
+	{
+		FUIAction UIAction;
+		{
+			// @todo: For safety, this should probably just copy the component id and try to find it, then toggle it.
+			UIAction.ExecuteAction = FExecuteAction::CreateLambda([&DebugInfoEntry](){ DebugInfoEntry.bShouldDraw = !DebugInfoEntry.bShouldDraw; });
+			UIAction.GetActionCheckState = FGetActionCheckState::CreateLambda([&DebugInfoEntry] { return DebugInfoEntry.bShouldDraw ? ECheckBoxState::Checked : ECheckBoxState::Unchecked; });
+		}
+
+		// @todo: Add an actor icon for the trajectory picker. Similar to RewindDebugger's.
+		Section.AddMenuEntry(DebugInfoEntry.Name, FText::FromName(DebugInfoEntry.Name), FText::FromName(DebugInfoEntry.Name), FSlateIcon(), UIAction, EUserInterfaceActionType::Check);
+	}
+}
+
+void FRewindDebuggerTrajectory::EnsureDebugDrawDelegateExists()
+{
+	if (!DebugDrawDelegateHandle.IsValid())
+	{
+		DebugDrawDelegateHandle = UDebugDrawService::Register(TEXT("GameplayDebug"), FDebugDrawDelegate::CreateRaw(this, &FRewindDebuggerTrajectory::DebugDraw));
+	}
+}
+
+void FRewindDebuggerTrajectory::DebugDraw(UCanvas* InCanvas, APlayerController* InController)
+{
+	// @todo: Only draw object ids which are part of the target object.
+	if (IsValid(WorldToVisualize) && !State.Trajectories.IsEmpty())
+	{
+		for (int TrajectoryIndex = 0; TrajectoryIndex < State.Trajectories.Num(); ++TrajectoryIndex)
+		{
+			if (State.DebugInfos[TrajectoryIndex].bShouldDraw)
+			{
+				const FGameplayTrajectory& DebugTrajectoryPair = State.Trajectories[TrajectoryIndex];
+				const FColor& DebugDrawColorForTrajectory = DebugDrawColors[TrajectoryIndex % DebugDrawColors.Num()];
+			
+				for (int SampleIndex = 0; SampleIndex < DebugTrajectoryPair.Samples.Num(); ++SampleIndex)
+				{
+					const FGameplayTrajectory::FSample& Sample = DebugTrajectoryPair.Samples[SampleIndex];
+
+					if (SampleIndex + 1 < DebugTrajectoryPair.Samples.Num())
+					{
+						const FGameplayTrajectory::FSample& FutureSample = DebugTrajectoryPair.Samples[SampleIndex + 1];
+
+						DrawDebugLine(WorldToVisualize, Sample.Position, FutureSample.Position, DebugDrawColorForTrajectory);
+					}
+
+					// TODO: Determine when to draw text. AKA every second, every X frames, etc.
+					if (SampleIndex % 10 == 0)
+					{
+						FVector2D SamplePositionOnScreen;
+					
+						if (InCanvas->SceneView->WorldToPixel(Sample.Position, SamplePositionOnScreen))
+						{
+							FFontRenderInfo Info;
+							Info.bEnableShadow = true;
+					
+							SamplePositionOnScreen.X = FMath::RoundToFloat(SamplePositionOnScreen.X);
+							SamplePositionOnScreen.Y = FMath::RoundToFloat(SamplePositionOnScreen.Y);
+						
+							InCanvas->DrawText(GEngine->GetTinyFont(), FString::SanitizeFloat(Sample.Time), SamplePositionOnScreen.X, SamplePositionOnScreen.Y, 1, 1, Info);
+						}
+					}
+				}
+			}
+		}
+	}
+}
+
+void FRewindDebuggerTrajectory::ShowBakeOutTrajectoryWindow()
+{
+	if (BakeOutWindow.IsValid())
+	{
+		FSlateApplication::Get().DestroyWindowImmediately(BakeOutWindow.ToSharedRef());
+	}
+	
+	// Custom window to setup bake out settings and execute the action
+	SAssignNew(BakeOutWindow, SWindow)
+	.Title(LOCTEXT("BakeOutTrajectoryWindowTitle", "Bake Out Trajectories"))
+	.SizingRule(ESizingRule::UserSized)
+	.SupportsMaximize(true)
+	.SupportsMinimize(true)
+	.MinWidth(850)
+	.MinHeight(500)
+	.CreateTitleBar(true)
+	.HasCloseButton(true)
+	.ClientSize(FVector2f{1280, 720});
+
+	// Gather all the associated owner names
+	TArray<FName> OwnerNames;
+	for (const FExtensionState::FDebugInfo& DebugInfo : State.DebugInfos)
+	{
+		OwnerNames.Add(DebugInfo.Name);
+	}
+
+	// Init custom trajectory export window
+	TSharedPtr<SExportTrajectoriesWindow> ExportWindow;
+	BakeOutWindow->SetContent
+	(
+		SAssignNew(ExportWindow, SExportTrajectoriesWindow)
+		.Trajectories(State.Trajectories)
+		.WidgetWindow(BakeOutWindow)
+		.OwnerNames(OwnerNames)
+		.DebugInfos(State.DebugInfos)
+		.SkelMeshInfos(State.SkelMeshInfos)
+	);
+	BakeOutWindow->SetOnWindowClosed(FOnWindowClosed::CreateLambda([this](const TSharedRef<SWindow>&){ BakeOutWindow = nullptr; }));
+
+	// Hook into main window as a modal window.
+	/*TSharedPtr<SWindow> ParentWindow;
+	if (FModuleManager::Get().IsModuleLoaded("MainFrame"))
+	{
+		const IMainFrameModule& MainFrame = FModuleManager::LoadModuleChecked<IMainFrameModule>("MainFrame");
+		ParentWindow = MainFrame.GetParentWindow();
+	}
+	ExportWindow->SetCanTick(true);*/
+	
+	// We use a modal window to ensure that we can store id/values from rewind debugger messages and ensure that its data provider are valid,
+	// otherwise we would have to copy all the information when opening the window.
+	// FSlateApplication::Get().AddModalWindow(Window, ParentWindow, false);
+	FSlateApplication::Get().AddWindow(BakeOutWindow.ToSharedRef());
+}
+
+void FRewindDebuggerTrajectory::BuildTrajectoryOwnersList(const IRewindDebugger* InRewindDebugger, const IGameplayProvider* InGameplayProvider, double InTraceStartTime, double InTraceEndTime, TArray<FObjectInfo> & OutObjectInfos)
+{
+	// Clear up previous info.
+	OutObjectInfos.Reset();
+	
+	// Used to filter which objects we can use to build our trajectory.
+	const FClassInfo* SkelMeshCmpClassInfo = InGameplayProvider->FindClassInfo(*USkeletalMeshComponent::StaticClass()->GetPathName());
+	check(SkelMeshCmpClassInfo)
+	
+	// Query target actor of the trace file / session.
+	InGameplayProvider->EnumerateObjects(InTraceStartTime, InTraceEndTime, [&OutObjectInfos,InGameplayProvider, SkelMeshCmpClassInfo, InRewindDebugger](const FObjectInfo& InObjectInfo)
+	{
+		// For all root objects find their skeletal mesh component's id, if any.
+		if (InObjectInfo.OuterId == 0)
+		{
+			InGameplayProvider->EnumerateSubobjects(InObjectInfo.Id, [&OutObjectInfos, InGameplayProvider, SkelMeshCmpClassInfo](uint64 InSubObjectId)
+			{
+				const FObjectInfo& SubObjectInfo = InGameplayProvider->GetObjectInfo(InSubObjectId);
+
+				const bool bIsSubObjectASkelMesh = SubObjectInfo.ClassId == SkelMeshCmpClassInfo->Id || InGameplayProvider->IsSubClassOf(SubObjectInfo.ClassId, SkelMeshCmpClassInfo->Id);
+				if (bIsSubObjectASkelMesh)
+				{
+					OutObjectInfos.Push(SubObjectInfo);
+				}
+			});
+		}
+		else
+		{
+			const bool bIsObjectASkelMesh = InObjectInfo.ClassId == SkelMeshCmpClassInfo->Id || InGameplayProvider->IsSubClassOf(InObjectInfo.ClassId, SkelMeshCmpClassInfo->Id);
+			if (bIsObjectASkelMesh)
+			{
+				const bool bHasOwningActor = InRewindDebugger->FindOwningActorInfo(InGameplayProvider, InObjectInfo.Id) != nullptr;
+				if (bHasOwningActor)
+				{
+					OutObjectInfos.Push(InObjectInfo);
+				}
+			}
+		}
+	});
+}
+
+void FRewindDebuggerTrajectory::BuildTrajectorySkeletalMeshInfoList(const IGameplayProvider* InGameplayProvider, const IAnimationProvider* InAnimationProvider, double InTraceStartTime, double InTraceEndTime, const TArray<FObjectInfo>& InSkeletalMeshComponentObjectInfos, TArray<FSkeletalMeshInfo>& OutSkeletalMeshInfos)
+{
+	// Clear up previous info.
+	OutSkeletalMeshInfos.Reset();
+	
+	for (int SkelMeshCmpIndex = 0; SkelMeshCmpIndex < InSkeletalMeshComponentObjectInfos.Num(); ++SkelMeshCmpIndex)
+	{
+		const FObjectInfo& SkeletalMeshComponentInfo = InSkeletalMeshComponentObjectInfos[SkelMeshCmpIndex];
+		
+		InAnimationProvider->ReadSkeletalMeshPoseTimeline(SkeletalMeshComponentInfo.Id, [InAnimationProvider, InGameplayProvider, InTraceStartTime, InTraceEndTime, SkelMeshCmpIndex, &SkeletalMeshComponentInfo, &OutSkeletalMeshInfos](const IAnimationProvider::SkeletalMeshPoseTimeline& TimelineData, bool bHasCurves)
+		{
+			TimelineData.EnumerateEvents(InTraceStartTime, InTraceEndTime, [InAnimationProvider, InGameplayProvider, SkelMeshCmpIndex, &SkeletalMeshComponentInfo, &OutSkeletalMeshInfos](double InStartTime, double InEndTime, uint32 InDepth, const FSkeletalMeshPoseMessage& InPoseMessage) -> TraceServices::EEventEnumerate
+			{
+				const FSkeletalMeshInfo* SkeletalMeshInfo = InAnimationProvider->FindSkeletalMeshInfo(InPoseMessage.MeshId);
+				const FObjectInfo* SkeletalMeshObjectInfo = InGameplayProvider->FindObjectInfo(InPoseMessage.MeshId); // @todo: Check if this search can be remove and just use the already available id.
+
+				// Query skeletal mesh information.
+				if (!SkeletalMeshInfo || !SkeletalMeshObjectInfo)
+				{
+					OutSkeletalMeshInfos.Reset();
+					UE_LOGFMT(LogTemp, Warning, "No SkelMeshInfo or SkelMeshObjectInfo provided. For object: {name}", *SkeletalMeshComponentInfo.Name);
+					return TraceServices::EEventEnumerate::Stop;
+				}
+				
+				if (!OutSkeletalMeshInfos.IsValidIndex(SkelMeshCmpIndex))
+				{
+					OutSkeletalMeshInfos.Push(*SkeletalMeshInfo);
+					// UE_LOGFMT(LogTemp, Warning, "[{0}] - Skel Mesh Info START - \nStarting:{1}", SkelMeshCmpIndex, SkeletalMeshObjectInfo->PathName);
+
+					const FObjectInfo* SkeletonObjectInfo = InGameplayProvider->FindObjectInfo(SkeletalMeshInfo->SkeletonId);
+					// UE_LOGFMT(LogTemp, Warning, "[{0}] - Skeleton Info START - \nStarting:{1}", SkelMeshCmpIndex, SkeletonObjectInfo != nullptr? SkeletonObjectInfo->PathName : static_cast<const TCHAR*>(TEXT("")));
+				}
+				else
+				{
+					// Keep track of any data change.
+					FSkeletalMeshInfo& PrevMeshInfo = OutSkeletalMeshInfos[SkelMeshCmpIndex];
+					
+					bool bSameMesh = SkeletalMeshInfo->Id == PrevMeshInfo.Id; 
+					bool bSameSkeleton = SkeletalMeshInfo->SkeletonId == PrevMeshInfo.SkeletonId;
+					
+					if (!bSameMesh)
+					{
+						UE_LOGFMT(LogTemp, Warning, "[{0}] - Skel Mesh Info Changed - \nPrev:{1}\nNew:{2}\nName:{3}", SkelMeshCmpIndex, PrevMeshInfo.Id, SkeletalMeshInfo->Id, SkeletalMeshObjectInfo->PathName);
+						OutSkeletalMeshInfos[SkelMeshCmpIndex] = *SkeletalMeshInfo;
+						
+						// If mesh changes is fine but not if the skeleton changes!
+						// UE_LOGFMT(LogTemp, Warning, "No SkelMeshInfo or SkelMeshObjectInfo provided. For object: {name}", SkeletalMeshObjectInfo->Name);
+					}
+					
+					if (!bSameSkeleton)
+					{
+						const FObjectInfo* SkeletonObjectInfo = InGameplayProvider->FindObjectInfo(SkeletalMeshInfo->SkeletonId);
+						UE_LOGFMT(LogTemp, Warning, "[{0}] - Skeleton Info Changed - \nPrev:{1}\nNew:{2}\nName:{3}", SkelMeshCmpIndex, PrevMeshInfo.SkeletonId, SkeletalMeshInfo->SkeletonId, SkeletonObjectInfo != nullptr ? SkeletonObjectInfo->PathName : static_cast<const TCHAR*>(TEXT("")));
+					}
+				}
+
+				return TraceServices::EEventEnumerate::Continue;
+			});
+		});
+	}
+}
+
+void FRewindDebuggerTrajectory::BuildTrajectories(const IGameplayProvider* InGameplayProvider, const IAnimationProvider* InAnimationProvider, double InTraceStartTime, double InTraceEndTime, const TArray<FObjectInfo>& InObjectInfos, TArray<FGameplayTrajectory>& OutTrajectories)
+{
+	// Clear up previous info.
+	OutTrajectories.Reset();
+
+	// Used to filter which objects we can use to build our trajectory.
+	const FClassInfo* SkelMeshCmpClassInfo = InGameplayProvider->FindClassInfo(*USkeletalMeshComponent::StaticClass()->GetPathName());
+	check(SkelMeshCmpClassInfo)
+
+	// Build trajectory from pose's root transform over time.
+	for (const FObjectInfo& ObjectInfo : InObjectInfos)
+	{
+		OutTrajectories.Push({});
+
+		// Setup initial index range
+		{
+			TRange<int32> IndexRange(TRangeBound<int32>::Inclusive(0), TRangeBound<int32>::Open());
+
+			OutTrajectories.Last().TraceInfo.Ranges.Push(IndexRange);
+		}
+		
+		// BoneCount split tracker.
+		int64 PrevSampleBoneCount = INDEX_NONE;
+		
+		InAnimationProvider->ReadSkeletalMeshPoseTimeline(ObjectInfo.Id, [InGameplayProvider, InAnimationProvider, InTraceStartTime, InTraceEndTime, &OutTrajectories, &ObjectInfo, &PrevSampleBoneCount](const IAnimationProvider::SkeletalMeshPoseTimeline& InTimeline, bool bHasCurves)
+		{
+			InTimeline.EnumerateEvents(InTraceStartTime, InTraceEndTime, [InGameplayProvider, InAnimationProvider, &OutTrajectories, &ObjectInfo, &PrevSampleBoneCount](double InStartTime, double InEndTime, uint32 InDepth, const FSkeletalMeshPoseMessage& InPoseMessage)
+			{
+				// Used to ensure trajectories always start at time being zero.
+				const bool bIsFirstSample = OutTrajectories.Last().Samples.IsEmpty();
+				const double FirstSampleRecordingTime = bIsFirstSample ? InPoseMessage.RecordingTime : OutTrajectories.Last().Samples[0].Time;
+				
+				const FVector SamplePosition = InPoseMessage.ComponentToWorld.GetTranslation();		// Root transform's position (SkelMeshCmp World Position).
+				const FQuat SampleOrientation = InPoseMessage.ComponentToWorld.GetRotation();		// Root transform's position (SkelMeshCmp World Rotation).
+				const double SampleTime = InPoseMessage.RecordingTime - FirstSampleRecordingTime;	// Sample timed, based of range beginning not trace beginning time.
+
+				// Query skeletal mesh information.
+				const FSkeletalMeshInfo* SkeletalMeshInfo = InAnimationProvider->FindSkeletalMeshInfo(InPoseMessage.MeshId);
+				const FObjectInfo* SkeletalMeshObjectInfo = InGameplayProvider->FindObjectInfo(InPoseMessage.MeshId);
+				
+				if (!SkeletalMeshInfo || !SkeletalMeshObjectInfo)
+				{
+					OutTrajectories.Last().Reset();
+					
+					UE_LOGFMT(LogTemp, Warning, "No SkelMeshInfo or SkelMeshObjectInfo provided. For object: {name}", *ObjectInfo.Name);
+					return TraceServices::EEventEnumerate::Stop;
+				}
+				
+				if (InPoseMessage.NumTransforms != SkeletalMeshInfo->BoneCount)
+				{
+					OutTrajectories.Last().Reset();
+					
+					UE_LOGFMT(LogTemp, Warning, "BoneCount doesn't match. Traced bone transforms do not match the respective traced mesh asset's bones. For object: {name}", *ObjectInfo.Name);
+					return TraceServices::EEventEnumerate::Stop;
+				}
+				
+				// Keep track of instances when bone counts don't match.
+				if (PrevSampleBoneCount != SkeletalMeshInfo->BoneCount)
+				{
+					int32 BoneCountSplitSample = OutTrajectories.Last().Samples.Num();
+
+					// Split range by the new bone count, unless we're initializing.
+					if (PrevSampleBoneCount != INDEX_NONE)
+					{
+						TRange<int32> LastRange = OutTrajectories.Last().TraceInfo.Ranges.Last();
+						OutTrajectories.Last().TraceInfo.Ranges.Pop();
+					
+						OutTrajectories.Last().TraceInfo.Ranges.Append(LastRange.Split(BoneCountSplitSample));
+					}
+					
+					// Assign the skeletal mesh info to the new range
+					OutTrajectories.Last().TraceInfo.SkeletalMeshInfos.Add(*SkeletalMeshInfo);
+				}
+				PrevSampleBoneCount = SkeletalMeshInfo->BoneCount;
+
+				// Extract and append animation pose for the current sample.
+				{
+					FTransform OutComponentToWorld;
+					OutTrajectories.Last().Poses.AddDefaulted();
+					InAnimationProvider->GetSkeletalMeshComponentSpacePose(InPoseMessage, *SkeletalMeshInfo, OutComponentToWorld, OutTrajectories.Last().Poses.Last());	
+				}
+
+				// Add current sample information.
+				OutTrajectories.Last().Samples.Push({SampleTime, SamplePosition, SampleOrientation});
+
+				return TraceServices::EEventEnumerate::Continue;
+			});
+		});
+
+		// Always set the upper bound range.
+		if (!OutTrajectories.Last().Samples.IsEmpty())
+		{
+			OutTrajectories.Last().TraceInfo.Ranges.Last().SetUpperBound(TRangeBound<int32>(OutTrajectories.Last().Samples.Num() - 1));
+		}
+		else
+		{
+			OutTrajectories.Last().TraceInfo.Reset();
+		}
+	}
+
+	// @todo: This should report back an error if it failed to perform an action.
+}
+
+void FRewindDebuggerTrajectory::UpdateDebugInfos(const IGameplayProvider* InGameplayProvider, const TArray<FObjectInfo>& InObjectInfos, TArray<FExtensionState::FDebugInfo>& InOutDebugInfos)
+{
+	// We need a local copy to search for previous occurances.
+	TArray<FExtensionState::FDebugInfo> PrevDebugInfos = InOutDebugInfos;
+	InOutDebugInfos.Reset();
+	
+	for (const FObjectInfo& ObjectInfo : InObjectInfos)
+	{
+		const FExtensionState::FDebugInfo* MatchedDebugInfo = PrevDebugInfos.FindByPredicate([&ObjectInfo](const FExtensionState::FDebugInfo& Item)
+		{
+			return Item.OwnerId == ObjectInfo.Id;
+		});
+
+		// Use cached version to keep track of stateful date (i.e. ShouldDraw checkbox)
+		if (MatchedDebugInfo)
+		{
+			InOutDebugInfos.Push(*MatchedDebugInfo);
+		}
+		else
+		{
+			InOutDebugInfos.Push({ObjectInfo.Id, GetFullNameForDebugInfoOwner(InGameplayProvider, ObjectInfo.Id), false});
+		}
+	}
+}
+
+FName FRewindDebuggerTrajectory::GetFullNameForDebugInfoOwner(const IGameplayProvider* InGameplayProvider, uint64 InOwnerObjectId)
+{
+	const FClassInfo* ActorClassInfo = InGameplayProvider->FindClassInfo(*AActor::StaticClass()->GetPathName());
+
+	uint64 ObjectId = InOwnerObjectId;
+
+	TStringBuilder<128> Result;
+	
+	while(true)
+	{
+		const FObjectInfo& ObjectInfo = InGameplayProvider->GetObjectInfo(ObjectId);
+
+		// We reached the owning actor.
+		if (InGameplayProvider->IsSubClassOf(ObjectInfo.ClassId, ActorClassInfo->Id))
+		{
+			Result.Prepend(TEXTVIEW(" - "));
+			Result.Prepend(ObjectInfo.Name);
+
+			const FWorldInfo* WorldInfo = InGameplayProvider->FindWorldInfoFromObject(ObjectId);
+			if (WorldInfo)
+			{
+				bool bIsServer = WorldInfo->NetMode == FWorldInfo::ENetMode::DedicatedServer;
+
+				if (bIsServer)
+				{
+					Result.Append(TEXTVIEW(" (Server)"));
+				}
+			}
+			
+			return FName(Result.ToView());
+		}
+		// We are a component, keep traversing up the tree.
+		else
+		{
+			if (ObjectInfo.OuterId != 0)
+			{
+				if (ObjectId != InOwnerObjectId)
+				{
+					Result.Prepend(TEXTVIEW(" - "));
+				}
+				Result.Prepend(ObjectInfo.Name);
+				
+				ObjectId = ObjectInfo.OuterId;
+			}
+			else
+			{
+				return NAME_None;
+			}
+		}
+	}
+}
+
+#undef LOCTEXT_NAMESPACE
diff --git a/Engine/Plugins/Experimental/Animation/TrajectoryTools/Source/TrajectoryTools/Private/TrajectoryToolsModule.cpp b/Engine/Plugins/Experimental/Animation/TrajectoryTools/Source/TrajectoryTools/Private/TrajectoryToolsModule.cpp
new file mode 100644
--- /dev/null
+++ b/Engine/Plugins/Experimental/Animation/TrajectoryTools/Source/TrajectoryTools/Private/TrajectoryToolsModule.cpp
@@ -0,0 +1,37 @@
+﻿// Copyright Epic Games, Inc. All Rights Reserved.
+
+#include "Modules/ModuleManager.h"
+#include "Modules/ModuleInterface.h"
+#include "TrajectoryRewindDebuggerExtension.h"
+#include "ToolMenus.h"
+#include "Toolkits/AssetEditorToolkit.h"
+
+#define LOCTEXT_NAMESPACE "FTrajectoryEditorModule"
+
+class FTrajectoryToolsModule : public IModuleInterface
+{
+private:
+	virtual void StartupModule() override;
+	virtual void ShutdownModule() override;
+
+	FRewindDebuggerTrajectory RewindDebuggerTrajectoryExtension;
+};
+
+void FTrajectoryToolsModule::StartupModule()
+{
+	if (GIsEditor && !IsRunningCommandlet())
+	{
+		RewindDebuggerTrajectoryExtension.Initialize();
+		IModularFeatures::Get().RegisterModularFeature(IRewindDebuggerExtension::ModularFeatureName, &RewindDebuggerTrajectoryExtension);
+	}
+}
+
+void FTrajectoryToolsModule::ShutdownModule()
+{
+	RewindDebuggerTrajectoryExtension.Shutdown();
+	IModularFeatures::Get().UnregisterModularFeature(IRewindDebuggerExtension::ModularFeatureName, &RewindDebuggerTrajectoryExtension);
+}
+
+IMPLEMENT_MODULE(FTrajectoryToolsModule, TrajectoryTools)
+
+#undef LOCTEXT_NAMESPACE
diff --git a/Engine/Plugins/Experimental/Animation/TrajectoryTools/Source/TrajectoryTools/Public/SExportTrajectoriesWindow.h b/Engine/Plugins/Experimental/Animation/TrajectoryTools/Source/TrajectoryTools/Public/SExportTrajectoriesWindow.h
new file mode 100644
--- /dev/null
+++ b/Engine/Plugins/Experimental/Animation/TrajectoryTools/Source/TrajectoryTools/Public/SExportTrajectoriesWindow.h
@@ -0,0 +1,237 @@
+﻿// Copyright Epic Games, Inc. All Rights Reserved.
+
+#pragma once
+
+#include "CoreMinimal.h"
+#include "TrajectoryLibrary.h"
+#include "Widgets/SCompoundWidget.h"
+#include "SEditorViewport.h"
+#include "PreviewScene.h"
+#include "EditorViewportClient.h"
+#include "TrajectoryRewindDebuggerExtension.h" // @todo: extract nested class and forward declare
+// #include "IAnimationProvider.h"
+
+#include "SExportTrajectoriesWindow.generated.h"
+
+class SMenuAnchor;
+class UAnimSequence;
+class SWindow;
+class ITableRow;
+class STableViewBase;
+class IDetailsView;
+class FSceneView;
+class FCanvas;
+class FViewport;
+class UDebugSkelMeshComponent;
+
+struct FSkeletalMeshInfo;
+
+// @todo: have a UTrajectoryViewportSettings, to store config / flags.
+
+UCLASS()
+class UTrajectoryExportDetails : public UObject
+{
+	GENERATED_BODY()
+
+public:
+	
+	/** Used to determine how the trajectory will be transformed into a animation sequence */
+	UPROPERTY(EditAnywhere, Category = "Export Settings", meta=(ShowOnlyInnerProperties))
+	FTrajectoryExportSettings ExportSettings;
+	
+	/** Number of key frames the animation sequence will have */
+	UPROPERTY(VisibleAnywhere, Category = "Information")
+	int NumberOfKeyFrames = 0;
+
+	/** Play length of the sequence obtained after backing trajectory */
+	UPROPERTY(VisibleAnywhere, Category = "Information")
+	double PlayLength = 0;
+
+	/** Used to create an asset for the trajectory to export its data on to. */
+	UPROPERTY(EditAnywhere, Category = "Output Asset", meta=(ShowOnlyInnerProperties))
+	FTrajectoryExportAssetInfo ExportAssetInfo;
+	
+	/** Reset all properties back to default values */
+	void Reset()
+	{
+		const UTrajectoryExportDetails* DefaultExportDetails = GetDefault<UTrajectoryExportDetails>();
+		check(DefaultExportDetails)
+		
+		ExportSettings.Reset();
+		ExportAssetInfo.Reset();
+		
+		// OutputAsset = DefaultExportDetails->OutputAsset;
+		PlayLength = DefaultExportDetails->PlayLength;
+		NumberOfKeyFrames = DefaultExportDetails->NumberOfKeyFrames;
+	}
+
+	// virtual void PostEditChangeProperty(FPropertyChangedEvent& PropertyChangedEvent) override;
+};
+
+/** Holds all the information needed to display a trajectory in the preview viewport */
+struct FPreviewData
+{
+	/** Index to access all the information related to the currently selected trajectory */
+	int SourceIndex = INDEX_NONE;
+
+	struct FViewportInfo
+	{
+		uint64 SourceSkeletalMeshComponentId = 0;
+		FSkeletalMeshInfo SourceSkeletalMeshInfo;
+	} Viewport;
+
+	/** Raw trajectories available */
+	TConstArrayView<FGameplayTrajectory> SourceTrajectories;
+	
+	/** Preview trajectory which is affected by the export settings */
+	FGameplayTrajectory OutputTrajectory;
+
+	/** We need to hold a pointer to be able to query the export settings and accurately display a preview trajectory */
+	TObjectPtr<UTrajectoryExportDetails> Details = nullptr; // @todo: this should be a tweak object ptr, tbh.
+
+	/** Holds stateful data used to draw timeline */
+	struct FTimelineInfo
+	{
+		bool bShouldDrawScrubber = false;
+		double ScrubTime = 0;
+	} Timeline;
+};
+
+/** Used to preview the trajectory to export given its settings (ExportRange,ExportFrameRage, etc). */
+class SPreviewTrajectoryViewport : public SEditorViewport
+{
+	
+public:
+	
+	SLATE_BEGIN_ARGS(SPreviewTrajectoryViewport) {}
+		SLATE_ARGUMENT(TSharedPtr<FPreviewData>, PreviewData)
+	SLATE_END_ARGS()
+
+	SPreviewTrajectoryViewport() : PreviewScene(FPreviewScene::ConstructionValues()), PreviewData(nullptr) {};
+	
+	void Construct(const FArguments& InArgs);
+
+	/** End SWidget */
+	virtual FReply OnKeyDown(const FGeometry& MyGeometry, const FKeyEvent& InKeyEvent) override;
+	/** End SWidget */
+
+	void ResetPreviewSkeletalMesh() const;
+	void UpdatePreviewPoseFromScrubTime(const IAnimationProvider* InAnimationProvider, const IGameplayProvider* InGameplayProvider, double InTraceStartTime, double InTraceEndTime) const;
+	
+	TSharedPtr<FPreviewData> GetPreviewData() const { return PreviewData; }
+
+protected:
+	/** Begin SEditorViewport interface */
+	virtual TSharedRef<FEditorViewportClient> MakeEditorViewportClient() override;
+	virtual TSharedPtr<SWidget> MakeViewportToolbar() override;
+	/** End SEditorViewport interface */
+
+private:
+	
+	/** Scene to preview trajectory on */
+	FPreviewScene PreviewScene;
+
+	/** Used to draw the trajectory */
+	TSharedPtr<FPreviewData> PreviewData;
+
+	/** Used to display the pose of the scrub time */
+	TObjectPtr<UDebugSkelMeshComponent> PreviewComponent;
+};
+
+/** Client for SPreviewTrajectoryViewport */
+class FPreviewTrajectoryViewportClient : public FEditorViewportClient
+{
+	
+public:
+	
+	FPreviewTrajectoryViewportClient(FPreviewScene& InPreviewScene, const TSharedRef<SPreviewTrajectoryViewport>& InPreviewTrajectoryViewport);
+
+	/** Begin FEditorViewportClient interface */
+	virtual void Tick(float DeltaTime) override;
+	virtual FSceneInterface* GetScene() const override { return PreviewScene->GetScene(); }
+	virtual FLinearColor GetBackgroundColor() const override { return FLinearColor(0.36f, 0.36f, 0.36f, 1); }
+	virtual void SetViewMode(EViewModeIndex Index) override final { FEditorViewportClient::SetViewMode(Index); }
+	virtual void Draw(const FSceneView* InView, FPrimitiveDrawInterface* InPDI) override;
+	virtual void DrawCanvas(FViewport& InViewport, FSceneView& InView, FCanvas& InCanvas) override;
+	/** End FEditorViewportClient interface */
+};
+
+/** Widget that configures and export a trajectory(s) into an animation asset / sequence */
+class SExportTrajectoriesWindow : public SCompoundWidget, FGCObject
+{
+public:
+
+	SLATE_BEGIN_ARGS(SExportTrajectoriesWindow) {}
+		SLATE_ARGUMENT(TSharedPtr<SWindow>, WidgetWindow)
+		SLATE_ARGUMENT(TArrayView<FName>, OwnerNames)
+		SLATE_ARGUMENT(TArrayView<FRewindDebuggerTrajectory::FExtensionState::FDebugInfo>, DebugInfos)
+		SLATE_ARGUMENT(TArrayView<FGameplayTrajectory>, Trajectories)
+		SLATE_ARGUMENT(TArrayView<FSkeletalMeshInfo>, SkelMeshInfos)
+	SLATE_END_ARGS()
+
+	SExportTrajectoriesWindow() {}
+
+	/** Usual Construct method of SWidgets. Constructs the slate UI and assigns the UI delegates. */
+	void Construct(const FArguments& InArgs);
+
+	/** Begin FGCObject */
+	virtual FString GetReferencerName() const override	{ return TEXT("SExportTrajectoriesWindow"); }
+	virtual void AddReferencedObjects(FReferenceCollector& Collector) override;
+	/** End FGCObject */
+
+	void ResetPreviewSelection()
+	{
+		SelectedTrajectoryIndex = INDEX_NONE;
+	}
+	
+private:
+	
+	// CALLBACKS
+	
+	TSharedRef<SWidget> OnGetTrajectoryPickerMenuContent();
+	void OnTrajectoryPickerMenuOpened(bool bIsOpen);
+	void OnFinishedChangingExportSettingsSelectionProperties(const FPropertyChangedEvent& PropertyChangedEvent);
+
+	// HELPER
+
+	void UpdateAssetInfo(FSoftObjectPath&  OutSkeletonPath, FSoftObjectPath & OutSkeletalMeshPath) const;
+	
+	/** Window owning this widget */
+	TWeakPtr<SWindow> WidgetWindow;
+
+	// INPUT DATA (Constant)
+
+	struct ImmutableState
+	{
+		// TArrayView<FGameplayTrajectory> Trajectories;
+		TArrayView<FRewindDebuggerTrajectory::FExtensionState::FDebugInfo> DebugInfos;
+		TArrayView<FSkeletalMeshInfo> SkeletalMeshInfos;
+		// TArrayView<FName> OwnerNames;
+	} ImmutableState;
+	
+	/** Trajectories that are available for export */
+	TArray<FGameplayTrajectory> Trajectories;
+
+	/** Names of the trajectories' associated objects */
+	TArray<FName> TrajectoryOwnerNames;
+
+	// SELECTION DATA
+
+	/** Index of trajectory to be exported */
+	int SelectedTrajectoryIndex = INDEX_NONE;
+	
+	/** Used to determine how the trajectory should be exported / baked out */
+	TObjectPtr<UTrajectoryExportDetails> ExportDetails;
+	
+	// PREVIEW DATA
+	
+	TSharedPtr<FPreviewData> PreviewData;
+	
+	// WIDGETS
+	
+	TSharedPtr<IDetailsView> ExportDetailsView;
+	TSharedPtr<SPreviewTrajectoryViewport> Viewport;
+	TSharedPtr<SMenuAnchor> TrajectoryPickerComboButton;
+	TSharedPtr<SWidget> ExportButton;
+	TSharedPtr<SWidget> ViewportTimeline;
+};
diff --git a/Engine/Plugins/Experimental/Animation/TrajectoryTools/Source/TrajectoryTools/Public/TrajectoryExportOperation.h b/Engine/Plugins/Experimental/Animation/TrajectoryTools/Source/TrajectoryTools/Public/TrajectoryExportOperation.h
new file mode 100644
--- /dev/null
+++ b/Engine/Plugins/Experimental/Animation/TrajectoryTools/Source/TrajectoryTools/Public/TrajectoryExportOperation.h
@@ -0,0 +1,139 @@
+// Copyright Epic Games, Inc. All Rights Reserved.
+
+#pragma once
+
+#include "CoreMinimal.h"
+#include "Misc/ScopedSlowTask.h"
+#include "Misc/FrameRate.h"
+#include "ScopedTransaction.h"
+#include "TrajectoryExportOperation.generated.h"
+
+struct FGameplayTrajectory;
+class UAnimSequence;
+
+USTRUCT()
+struct TRAJECTORYTOOLS_API FTrajectoryExportSettings
+{
+	GENERATED_BODY()
+	
+	FTrajectoryExportSettings();
+	
+	/** Frame rate for the exported asset */
+	UPROPERTY(EditAnywhere, Category = "Settings")
+	FFrameRate FrameRate;
+	
+	/** Range of the trajectory to export. Note that the entire trajectory will be exported if the range is empty (Min == Max). */
+	UPROPERTY(EditAnywhere, meta=(Units="s"), Category = "Settings")
+	FFloatInterval Range;
+
+	/** Flag to determine if all trajectory samples should be relative to the frame given by ForceOriginAtTime*/
+	UPROPERTY(EditAnywhere, meta=(DisplayName="Force Origin"), Category = "Settings")
+	bool bShouldForceOrigin = false;
+	
+	/** Specific which trajectory sample's position will be position at the origin given by the desired time. Note that if time for origin is less than 0 no origin will be enforced. */
+	UPROPERTY(EditAnywhere, meta=(EditCondition="bShouldForceOrigin", EditConditionHides, Units="s"), Category = "Settings")
+	double OriginTime = 0;
+
+	/**
+	 * Any files with the same name will be overwritten instead of creating a new file with a numeric suffix.
+	 * This is useful when iterating on a batch process.
+	*/
+	UPROPERTY(EditAnywhere, meta=(DisplayName="Overwrite Existing Files"), Category = "Settings")
+	bool bShouldOverwriteExistingFiles = false;
+
+	/** Will not produce keys on bones that are not animated, reducing size on disk of the resulting files. */
+	UPROPERTY(EditAnywhere, meta=(DisplayName="Export Only Animated Bones"), Category = "Settings")
+	bool bShouldExportOnlyAnimatedBones = true;
+	
+	/** All export settings are valid and trajectory can be generated without issues from them */
+	bool IsValid() const;
+
+	/** Reset trajectory to default values */
+	void Reset();
+};
+
+/** All information needed to create an asset for a trajectory */
+USTRUCT()
+struct TRAJECTORYTOOLS_API FTrajectoryExportAssetInfo
+{
+	GENERATED_BODY()
+	
+	/** Name of asset to be created during export. Defaulted to the name of the selected trajectory in UI. */
+	UPROPERTY(EditAnywhere, Category = "Asset")
+	FString AssetName;
+	
+	/** Directory path where to export asset to */
+	UPROPERTY(EditAnywhere, Category = "Asset", meta=(ContentDir))
+	FDirectoryPath FolderPath;
+
+	/** Path to skeleton to be used and assigned to the exported asset. */
+	UPROPERTY(VisibleAnywhere, Category = "Information")
+	FSoftObjectPath Skeleton;
+
+	/** Path to skeletal mesh to be used and assigned to the exported asset. */
+	UPROPERTY(VisibleAnywhere, Category = "Information")
+	FSoftObjectPath SkeletalMesh;
+
+	bool CanCreateAsset() const;
+	
+	void Reset();
+
+	bool IsValid() const;
+};
+
+/** Data needed to run an "export" operation on a trajectory. */
+struct TRAJECTORYTOOLS_API FTrajectoryExportContext
+{
+	/** Used when transforming trajectory data before saving/exporting it to an asset. */
+	FTrajectoryExportSettings Settings;
+	
+	/** Used when creating the asset to hold the trajectory data. */
+	FTrajectoryExportAssetInfo AssetInfo;
+
+	/** Raw trajectory data that will be export after operation is complete. */
+	FGameplayTrajectory* Data;
+
+	/** Name of the object that we sourced the trajectory data from */
+	FString SourceObjectName;
+	
+	void Reset();
+
+	bool IsValid() const;
+};
+
+/** Encapsulate ability to export a trajectory to the specified asset. */
+UCLASS()
+class TRAJECTORYTOOLS_API UTrajectoryExportOperation : public UObject
+{
+	GENERATED_BODY()
+	
+public:
+
+	// @todo: After having a general trajectory type, expose this to blueprint to allow for editor scripts, bp, etc to use this.
+	static void ExportTrajectory(FGameplayTrajectory* InTrajectory, const FTrajectoryExportSettings& InSettings, const FTrajectoryExportAssetInfo& InAssetInfo, const FString& InSourceObjectName);
+	
+	/** Actually run the process to export the trajectory for the given context. */
+	void Run(const FTrajectoryExportContext& Context);
+	
+private:
+
+	/** Clear any generated/saved info. */
+	void Reset();
+
+	/** Create assets to export the trajectory data to. */
+	void GenerateAssets(const FTrajectoryExportContext& Context, FScopedSlowTask& Progress);
+
+	/** Output trajectory data to their respective assets. */
+	void ExportDataToAssets(const FTrajectoryExportContext& Context, FScopedSlowTask& Progress) const;
+
+	/** Output notifications of results. */
+	void NotifyUserOfResults(const FTrajectoryExportContext& Context, FScopedSlowTask& Progress) const;
+	
+	/** If user cancelled half way, cleanup all created asset(s). */
+	void CleanupIfCancelled(const FScopedSlowTask& Progress) const;
+
+	TWeakObjectPtr<UAnimSequence> GeneratedAsset = nullptr;
+	TWeakObjectPtr<UAnimSequence> AssetToProcess = nullptr;
+
+	TUniquePtr<FScopedTransaction> ActiveTransaction = nullptr;
+};
diff --git a/Engine/Plugins/Experimental/Animation/TrajectoryTools/Source/TrajectoryTools/Public/TrajectoryLibrary.h b/Engine/Plugins/Experimental/Animation/TrajectoryTools/Source/TrajectoryTools/Public/TrajectoryLibrary.h
new file mode 100644
--- /dev/null
+++ b/Engine/Plugins/Experimental/Animation/TrajectoryTools/Source/TrajectoryTools/Public/TrajectoryLibrary.h
@@ -0,0 +1,87 @@
+// Copyright Epic Games, Inc. All Rights Reserved.
+
+#pragma once
+
+#include "CoreMinimal.h"
+#include "Misc/FrameRate.h"
+#include "TrajectoryExportOperation.h"
+#include "IAnimationProvider.h"
+#include "TrajectoryLibrary.generated.h"
+
+class UAnimSequence;
+struct FTrajectoryExportSettings;
+
+// @todo: Temp. This will be swapped for general purpose trajectory struct + poses array.
+USTRUCT()
+struct TRAJECTORYTOOLS_API FGameplayTrajectory
+{
+	GENERATED_BODY();
+	
+	struct FSample
+	{
+		double Time = 0;
+		FVector Position = FVector::ZeroVector;
+		FQuat Orientation = FQuat::Identity;
+	};
+	
+	TArray<FSample> Samples;
+	TArray<TArray<FTransform>> Poses;
+	
+	struct FTraceRangedBuffers
+	{
+		TArray<TRange<int32>> Ranges;
+		TArray<FSkeletalMeshInfo> SkeletalMeshInfos;
+
+		int Num() const
+		{
+			ensure(Ranges.Num() == SkeletalMeshInfos.Num());
+			return Ranges.Num();
+		}
+		
+		bool IsEmpty() const
+		{
+			return ensure(Ranges.IsEmpty() == SkeletalMeshInfos.IsEmpty()) && Ranges.IsEmpty();
+		}
+
+		void Reset(int NewSize = 0)
+		{
+			Ranges.Reset(NewSize);
+			SkeletalMeshInfos.Reset(NewSize);
+		}
+		
+	} TraceInfo;
+	
+	void Reset()
+	{
+		Samples.Reset();
+		Poses.Reset();
+		TraceInfo.Ranges.Reset();
+		TraceInfo.SkeletalMeshInfos.Reset();
+	}
+
+	bool IsValid() const
+	{
+		return Samples.Num() == Poses.Num();
+	}
+};
+
+struct TRAJECTORYTOOLS_API FTrajectoryToolsLibrary
+{
+	struct FRangeOverlapTestResult
+	{
+		bool bOverlaps = false;
+		TArray<int> Ranges;
+	};
+	
+	static void GetRangeOverlaps(const FGameplayTrajectory& InTrajectory, const TRange<int32>& InSampleRange, FRangeOverlapTestResult& InOutOverlapResult);
+	
+	static void GetSampleIndicesForMatchedSampleTime(const FGameplayTrajectory& InTrajectory, double InRequestedSampleTime, TPair<int32, int32>& OutMatchedSampleIndicesRange);
+
+	static bool GetTransformAtTimeInTrajectory(const FGameplayTrajectory& InTrajectory, double InRequestedTime, FTransform& OuTransform);
+
+	static bool GetPoseAtTimeInTrajectory(const FGameplayTrajectory& InTrajectory, double InRequestedTime, TArray<FTransform>& OutCSPose);
+	
+	static void TransformTrajectoryToMatchFrameRate(const FGameplayTrajectory& InTrajectory, FFrameRate InFrameRate, FGameplayTrajectory& OutFrameMatchedTrajectory);
+
+	static void TransformTrajectoryToMatchExportSettings(const FGameplayTrajectory& InTrajectory, const FTrajectoryExportSettings& InExportSettings, FGameplayTrajectory& OutTransformedTrajectory);
+};
\ No newline at end of file
diff --git a/Engine/Plugins/Experimental/Animation/TrajectoryTools/Source/TrajectoryTools/Public/TrajectoryRewindDebuggerExtension.h b/Engine/Plugins/Experimental/Animation/TrajectoryTools/Source/TrajectoryTools/Public/TrajectoryRewindDebuggerExtension.h
new file mode 100644
--- /dev/null
+++ b/Engine/Plugins/Experimental/Animation/TrajectoryTools/Source/TrajectoryTools/Public/TrajectoryRewindDebuggerExtension.h
@@ -0,0 +1,106 @@
+// Copyright Epic Games, Inc. All Rights Reserved.
+
+#pragma once
+
+#include "CoreMinimal.h"
+#include "IAnimationProvider.h"
+#include "IGameplayProvider.h"
+#include "IRewindDebugger.h"
+#include "IRewindDebuggerExtension.h"
+#include "TrajectoryLibrary.h"
+
+class APlayerController;
+class UCanvas;
+class UToolMenu;
+struct FToolMenuSection;
+class SWindow;
+
+/**
+ * A Rewind Debugger extension that allows user to visualize trajectories and export them from the current trace session.
+ */
+class FRewindDebuggerTrajectory : public IRewindDebuggerExtension
+{
+public:
+	
+	FRewindDebuggerTrajectory();
+	virtual ~FRewindDebuggerTrajectory() {};
+
+	void Initialize();
+	void Shutdown();
+	void Reset();
+
+	/** Begin IRewindDebuggerExtension implementation */
+	virtual void Update(float DeltaTime, IRewindDebugger* RewindDebugger) override;
+	virtual void RecordingStarted(IRewindDebugger* RewindDebugger) override;
+	virtual void RecordingStopped(IRewindDebugger* RewindDebugger) override;
+	virtual void Clear(IRewindDebugger* RewindDebugger) override;
+	virtual FString GetName() { return TEXT("RewindDebuggerTrajectory"); }
+	/* End IRewindDebuggerExtension implementation */
+	
+private:
+
+	// UI
+	
+	void GenerateMenu();
+	void MakeTrajectoriesMenu(FToolMenuSection& InSection);;
+	void MakeToggleDebugDrawMenu(UToolMenu* InMenu);
+	
+	// Debug Draw
+	void EnsureDebugDrawDelegateExists();
+	void DebugDraw(UCanvas* InCanvas, APlayerController* InController);
+
+	/** Holds information from each extension's update */
+	struct FExtensionState
+	{
+		/** Information for UI / debug drawing */
+		struct FDebugInfo
+		{
+			uint64 OwnerId = 0;
+			FName Name = NAME_None;
+			bool bShouldDraw = false;
+		};
+
+		/** Information to extract a trajectory/clip from a trace file */
+		struct FExtract
+		{
+			double TraceStarTime = 0;
+			double TraceEndTime = 0;
+		} Extract;
+		
+		TArray<FGameplayTrajectory> Trajectories;	// Trajectories extracted from trace session
+		TArray<FObjectInfo> ObjectInfos;			// Associated object's information for the extracted trajectories
+		TArray<FSkeletalMeshInfo> SkelMeshInfos;	// Associated object's skeletal mesh info for the extracted trajectories.	
+		TArray<FDebugInfo> DebugInfos;				// Associated debug information for extracted trajectories
+
+		bool bWasPIESimulating = false;				// Used to request a trajectory extraction operation.
+		
+		/** Reset extension's state */
+		void Reset()
+		{
+			Extract.TraceEndTime = 0;
+			Extract.TraceStarTime = 0;
+			
+			Trajectories.Reset();
+			ObjectInfos.Reset();
+			SkelMeshInfos.Reset();
+			DebugInfos.Reset();
+
+			bWasPIESimulating = false;
+		}
+	} State;
+	
+	TArray<FColor> DebugDrawColors;
+	FDelegateHandle DebugDrawDelegateHandle;
+	UWorld* WorldToVisualize = nullptr;
+	TSharedPtr<SWindow> BakeOutWindow;
+	
+	void ShowBakeOutTrajectoryWindow();
+
+	static void BuildTrajectoryOwnersList(const IRewindDebugger* InRewindDebugger, const IGameplayProvider* InGameplayProvider, double InTraceStartTime, double InTraceEndTime, TArray<FObjectInfo> & OutObjectInfos);
+	static void BuildTrajectorySkeletalMeshInfoList(const IGameplayProvider* InGameplayProvider, const IAnimationProvider* InAnimationProvider, double InTraceStartTime, double InTraceEndTime, const TArray<FObjectInfo>& InSkeletalMeshComponentObjectInfos, TArray<FSkeletalMeshInfo>& OutSkeletalMeshInfos);
+	static void BuildTrajectories(const IGameplayProvider* InGameplayProvider, const IAnimationProvider* InAnimationProvider, double InTraceStartTime, double InTraceEndTime, const TArray<FObjectInfo>& InObjectInfos, TArray<FGameplayTrajectory>& OutTrajectories);
+	static void UpdateDebugInfos(const IGameplayProvider* InGameplayProvider, const TArray<FObjectInfo>& InObjectInfos, TArray<FExtensionState::FDebugInfo>& InOutDebugInfos);
+	static FName GetFullNameForDebugInfoOwner(const IGameplayProvider* InGameplayProvider, uint64 InOwnerObjectId);
+
+	friend class SExportTrajectoriesWindow;
+};
\ No newline at end of file
diff --git a/Engine/Plugins/Experimental/Animation/TrajectoryTools/Source/TrajectoryTools/TrajectoryTools.Build.cs b/Engine/Plugins/Experimental/Animation/TrajectoryTools/Source/TrajectoryTools/TrajectoryTools.Build.cs
new file mode 100644
--- /dev/null
+++ b/Engine/Plugins/Experimental/Animation/TrajectoryTools/Source/TrajectoryTools/TrajectoryTools.Build.cs
@@ -0,0 +1,71 @@
+﻿// Copyright Epic Games, Inc. All Rights Reserved.
+
+using UnrealBuildTool;
+
+public class TrajectoryTools : ModuleRules
+{
+	public TrajectoryTools(ReadOnlyTargetRules Target) : base(Target)
+	{
+		PCHUsage = ModuleRules.PCHUsageMode.UseExplicitOrSharedPCHs;
+		
+		PublicIncludePaths.AddRange(
+			new string[] {
+				// ... add public include paths required here ...
+			}
+			);
+				
+		
+		PrivateIncludePaths.AddRange(
+			new string[] {
+				// ... add other private include paths required here ...
+			}
+			);
+			
+		
+		PublicDependencyModuleNames.AddRange(
+			new string[]
+			{
+				"Core",
+				// ... add other public dependencies that you statically link with here ...
+			}
+			);
+			
+		
+		PrivateDependencyModuleNames.AddRange(
+			new string[]
+			{
+				"CoreUObject",
+				"Engine",
+				
+				// Trace-related dependencies
+				"TraceLog",
+				"TraceAnalysis",
+				"TraceServices",
+				"TraceInsights",
+				"GameplayInsights",
+				"RewindDebugger",
+				
+				// UI 
+				"PropertyEditor",
+				"Slate",
+				"SlateCore",
+				"RewindDebuggerInterface",
+				"UnrealEd",
+				"InputCore",
+				"ToolMenus",
+				"ToolWidgets",
+				"EditorWidgets", 
+				
+				// ... add private dependencies that you statically link with here ...	
+			}
+			);
+		
+		
+		DynamicallyLoadedModuleNames.AddRange(
+			new string[]
+			{
+				// ... add any modules that your module loads dynamically here ...
+			}
+			);
+	}
+}
diff --git a/Engine/Plugins/Experimental/Animation/TrajectoryTools/TrajectoryTools.uplugin b/Engine/Plugins/Experimental/Animation/TrajectoryTools/TrajectoryTools.uplugin
new file mode 100644
--- /dev/null
+++ b/Engine/Plugins/Experimental/Animation/TrajectoryTools/TrajectoryTools.uplugin
@@ -0,0 +1,31 @@
+﻿{
+	"FileVersion": 3,
+	"Version": 1,
+	"VersionName": "0.1",
+	"FriendlyName": "Trajectory Tools",
+	"Description": "Workflows for gameplay trajectories, such as extracting them from recordings.",
+	"Category": "Animation",
+	"CreatedBy": "Epic Games, Inc.",
+	"CreatedByURL": "",
+	"DocsURL": "",
+	"MarketplaceURL": "",
+	"SupportURL": "",
+	"EnabledByDefault": false,
+	"CanContainContent": false,
+	"IsBetaVersion": false,
+	"IsExperimentalVersion": true,
+	"Installed": false,
+	"Modules": [
+		{
+			"Name": "TrajectoryTools",
+			"Type": "Editor",
+			"LoadingPhase": "PostEngineInit"
+		}
+	],
+	"Plugins": [
+		{
+			"Name": "GameplayInsights",
+			"Enabled": true
+		}
+	]
+}
\ No newline at end of file
