Fix for conflicting localization keys

#rnx

diff --git a/Engine/Plugins/Experimental/Animation/TrajectoryTools/Source/TrajectoryTools/Private/SExportTrajectoriesWindow.cpp b/Engine/Plugins/Experimental/Animation/TrajectoryTools/Source/TrajectoryTools/Private/SExportTrajectoriesWindow.cpp
--- a/Engine/Plugins/Experimental/Animation/TrajectoryTools/Source/TrajectoryTools/Private/SExportTrajectoriesWindow.cpp
+++ b/Engine/Plugins/Experimental/Animation/TrajectoryTools/Source/TrajectoryTools/Private/SExportTrajectoriesWindow.cpp
@@ -784,7 +784,7 @@ void SExportTrajectoriesWindow::Construct(const FArguments& InArgs)
 						// Set max range via timeline
 						{
 							const FText Label = FText::Format(LOCTEXT("SetExportRangeMaxLabel", "Set export range max at frame {0}* ({1}s)"), ScrubSampleIndex, ScrubSampleTime); 
-							const FText ToolTip = LOCTEXT("SetExportRangeMinToolTip", "Sets the ExportRange's Max value using the current scrub frame's time.");
+							const FText ToolTip = LOCTEXT("SetExportRangeMaxToolTip", "Sets the ExportRange's Max value using the current scrub frame's time.");
 							FUIAction Action = FExecuteAction::CreateLambda([InPreviewData, InExportWindow]()
 							{
 								if (IsValid(InPreviewData->Details))
@@ -808,7 +808,7 @@ void SExportTrajectoriesWindow::Construct(const FArguments& InArgs)
 						// Set origin via timeline
 						{
 							const FText Label = FText::Format(LOCTEXT("SetThisFrameAsTheOriginLabel", "Set origin at frame {0}* ({1}s)"), ScrubSampleIndex, ScrubSampleTime); 
-							const FText ToolTip = LOCTEXT("SetExportRangeMinToolTip", "Sets the current scrub frame at the origin in baked out asset.");
+							const FText ToolTip = LOCTEXT("SetExportOriginToolTip", "Sets the current scrub frame at the origin in baked out asset.");
 							FUIAction Action = FExecuteAction::CreateLambda([InPreviewData]()
 							{
 								if (IsValid(InPreviewData->Details))
diff --git a/Engine/Plugins/Experimental/Animation/TrajectoryTools/Source/TrajectoryTools/Private/TrajectoryRewindDebuggerExtension.cpp b/Engine/Plugins/Experimental/Animation/TrajectoryTools/Source/TrajectoryTools/Private/TrajectoryRewindDebuggerExtension.cpp
--- a/Engine/Plugins/Experimental/Animation/TrajectoryTools/Source/TrajectoryTools/Private/TrajectoryRewindDebuggerExtension.cpp
+++ b/Engine/Plugins/Experimental/Animation/TrajectoryTools/Source/TrajectoryTools/Private/TrajectoryRewindDebuggerExtension.cpp
@@ -158,8 +158,8 @@ void FRewindDebuggerTrajectory::GenerateMenu()
 		Section.AddSeparator("TrajectoryWorkflows");
 		{
 			static const FName Name = "Trajectories";
-			static const FText Label = LOCTEXT("ToggleDebugDrawTrajectoriesLabel", "Trajectories");
-			static const FText ToolTip = LOCTEXT("ToggleDebugDrawTrajectoriesTooltip", "Toggle trajectories to display in viewport");
+			static const FText Label = LOCTEXT("TrajectoriesMenuLabel", "Trajectories");
+			static const FText ToolTip = LOCTEXT("TrajectoriesMenuTooltip", "Toggle trajectories to display in viewport");
 			
 			Section.AddSubMenu(Name, Label, ToolTip, FNewToolMenuDelegate::CreateLambda([this](UToolMenu* InMenu)
 			{
