Fixed crash were we can pick trajectories that have no skeletal mesh info

diff --git a/Engine/Plugins/Experimental/Animation/TrajectoryTools/Source/TrajectoryTools/Private/SExportTrajectoriesWindow.cpp b/Engine/Plugins/Experimental/Animation/TrajectoryTools/Source/TrajectoryTools/Private/SExportTrajectoriesWindow.cpp
--- a/Engine/Plugins/Experimental/Animation/TrajectoryTools/Source/TrajectoryTools/Private/SExportTrajectoriesWindow.cpp
+++ b/Engine/Plugins/Experimental/Animation/TrajectoryTools/Source/TrajectoryTools/Private/SExportTrajectoriesWindow.cpp
@@ -682,7 +682,7 @@ void SExportTrajectoriesWindow::Construct(const FArguments& InArgs)
 	ImmutableState.SkeletalMeshInfos = InArgs._SkelMeshInfos;
 	ImmutableState.DebugInfos = InArgs._DebugInfos;
 	
-	check(Trajectories.Num() == ImmutableState.SkeletalMeshInfos.Num() && ImmutableState.SkeletalMeshInfos.Num() == ImmutableState.DebugInfos.Num())	
+	check(Trajectories.Num() == ImmutableState.DebugInfos.Num())
 	
 	// Export configuration settings.
 	FName Name = FName("TrajectoryExportDetails");
@@ -973,7 +973,8 @@ TSharedRef<SWidget> SExportTrajectoriesWindow::OnGetTrajectoryPickerMenuContent(
 
 		UIAction.CanExecuteAction = FCanExecuteAction::CreateLambda([i, this]()
 		{
-			return Trajectories[i].IsValid();
+			const bool bHasAnyData = !Trajectories[i].Samples.IsEmpty() && !Trajectories[i].Poses.IsEmpty();
+			return Trajectories[i].IsValid() && bHasAnyData;
 		});
 		
 		UIAction.ExecuteAction = FExecuteAction::CreateLambda([i, this]()
diff --git a/Engine/Plugins/Experimental/Animation/TrajectoryTools/Source/TrajectoryTools/Private/TrajectoryExportOperation.cpp b/Engine/Plugins/Experimental/Animation/TrajectoryTools/Source/TrajectoryTools/Private/TrajectoryExportOperation.cpp
--- a/Engine/Plugins/Experimental/Animation/TrajectoryTools/Source/TrajectoryTools/Private/TrajectoryExportOperation.cpp
+++ b/Engine/Plugins/Experimental/Animation/TrajectoryTools/Source/TrajectoryTools/Private/TrajectoryExportOperation.cpp
@@ -342,6 +342,12 @@ void UTrajectoryExportOperation::ExportDataToAssets(const FTrajectoryExportConte
 			// Convert from component space to local space.
 			for (int32 BoneIndex = NumBones - 1; BoneIndex > 0; --BoneIndex)
 			{
+			    // @todo: temporary to avoid crashing due to VB not being taken into account since we are using raw parent indices. Need to properly handle this.
+				if (BoneIndex >= RefSkeleton.GetRawBoneNum())
+				{
+					continue;
+				}
+				
 				const int32 ParentIndex = RefSkeleton.GetRawParentIndex(BoneIndex);
 				if (ParentIndex != INDEX_NONE)
 				{
