Fixed using inproper reference skeleton when exporting trajectories which caused bone hierarchy to not be respected or crashed the editor.

diff --git a/Engine/Plugins/Experimental/Animation/TrajectoryTools/Source/TrajectoryTools/Private/TrajectoryExportOperation.cpp b/Engine/Plugins/Experimental/Animation/TrajectoryTools/Source/TrajectoryTools/Private/TrajectoryExportOperation.cpp
--- a/Engine/Plugins/Experimental/Animation/TrajectoryTools/Source/TrajectoryTools/Private/TrajectoryExportOperation.cpp
+++ b/Engine/Plugins/Experimental/Animation/TrajectoryTools/Source/TrajectoryTools/Private/TrajectoryExportOperation.cpp
@@ -227,20 +227,27 @@ void UTrajectoryExportOperation::ExportDataToAssets(const FTrajectoryExportConte
 		return;
 	}
 	
-	// A valid skeleton is needed to query the ref pose skeleton.
+	// A valid skeleton is needed for preview information.
 	USkeleton* Skeleton = TSoftObjectPtr<USkeleton>(Context.AssetInfo.Skeleton).LoadSynchronous();
 	if (!IsValid(Skeleton) || Skeleton->HasAnyFlags(RF_BeginDestroyed | RF_FinishDestroyed))
 	{
 		return;
 	}
 
+	// A valid skeletal mesh is needed to query the ref pose.
+	USkeletalMesh* SkeletalMesh = TSoftObjectPtr<USkeletalMesh>(Context.AssetInfo.SkeletalMesh).LoadSynchronous();
+	if (!IsValid(SkeletalMesh) || SkeletalMesh->HasAnyFlags(RF_BeginDestroyed | RF_FinishDestroyed))
+	{
+		return;
+	}
+
 	// No bones in ref skeleton. Can't convert traced poses from component to local space.
-	const FReferenceSkeleton& RefSkeleton = Skeleton->GetReferenceSkeleton();
+	const FReferenceSkeleton& RefSkeleton = SkeletalMesh->GetRefSkeleton();
 	if (RefSkeleton.GetRawBoneNum() == 0)
 	{
 		return;
 	}
-
+	
 	// Inconsistent tracing of samples vs poses. Something wrong happened.
 	const FGameplayTrajectory& RawTrajectory = *Context.Data;
 	if (RawTrajectory.Samples.Num() != RawTrajectory.Poses.Num())
@@ -304,8 +311,8 @@ void UTrajectoryExportOperation::ExportDataToAssets(const FTrajectoryExportConte
 		AssetToProcess->ImportResampleFramerate = static_cast<int32>(Context.Settings.FrameRate.AsInterval());
 
 		// Add all bone tracks in ref skeleton for now. @todo: Take into consideration bShouldExportOnlyAnimatedBones.
-		const int32 NumBones = RefSkeleton.GetRawBoneNum();
-		const TArray<FTransform>& RefBonePose = RefSkeleton.GetRawRefBonePose();
+		const int32 NumBones = RefSkeleton.GetRefBonePose().Num();
+		const TArray<FTransform>& RefBonePose = RefSkeleton.GetRefBonePose();
 		for (int32 BoneIndex = 0; BoneIndex < NumBones; ++BoneIndex)
 		{
 			const FName BoneName = RefSkeleton.GetBoneName(BoneIndex);
@@ -342,13 +349,13 @@ void UTrajectoryExportOperation::ExportDataToAssets(const FTrajectoryExportConte
 			// Convert from component space to local space.
 			for (int32 BoneIndex = NumBones - 1; BoneIndex > 0; --BoneIndex)
 			{
-			    // @todo: temporary to avoid crashing due to VB not being taken into account since we are using raw parent indices. Need to properly handle this.
-				if (BoneIndex >= RefSkeleton.GetRawBoneNum())
+			    // @todo: We should probably warn about bone mismatch and cancel operation instead if any of the traced poses doesn't match the number of bones in the references skeleton.
+				if (BoneIndex >= RefSkeleton.GetRefBonePose().Num())
 				{
 					continue;
 				}
 				
-				const int32 ParentIndex = RefSkeleton.GetRawParentIndex(BoneIndex);
+				const int32 ParentIndex = RefSkeleton.GetParentIndex(BoneIndex);
 				if (ParentIndex != INDEX_NONE)
 				{
 					LocalSpacePoses[PoseIndex][BoneIndex] = LocalSpacePoses[PoseIndex][BoneIndex].GetRelativeTransform(LocalSpacePoses[PoseIndex][ParentIndex]);
@@ -406,7 +413,7 @@ void UTrajectoryExportOperation::ExportDataToAssets(const FTrajectoryExportConte
 				}
 				else
 				{
-					const int32 TrackNameToBoneIndex = RefSkeleton.FindRawBoneIndex(TrackName);
+					const int32 TrackNameToBoneIndex = RefSkeleton.FindBoneIndex(TrackName);
 
 					if (TrackNameToBoneIndex != INDEX_NONE && LocalSpacePoses[KeyIndex].IsValidIndex(TrackNameToBoneIndex))
 					{
