Now we build trajectory information only when opening the tool window.

Fixed not using current trace time instead of total trace time when building trajectory information.

diff --git a/Engine/Plugins/Experimental/Animation/TrajectoryTools/Source/TrajectoryTools/Private/TrajectoryRewindDebuggerExtension.cpp b/Engine/Plugins/Experimental/Animation/TrajectoryTools/Source/TrajectoryTools/Private/TrajectoryRewindDebuggerExtension.cpp
--- a/Engine/Plugins/Experimental/Animation/TrajectoryTools/Source/TrajectoryTools/Private/TrajectoryRewindDebuggerExtension.cpp
+++ b/Engine/Plugins/Experimental/Animation/TrajectoryTools/Source/TrajectoryTools/Private/TrajectoryRewindDebuggerExtension.cpp
@@ -79,11 +79,6 @@ void FRewindDebuggerTrajectory::Update(float DeltaTime, IRewindDebugger* RewindD
 	{
 		Reset();
 
-		if (RewindDebugger->IsPIESimulating())
-		{
-			State.bWasPIESimulating = true;
-		}
-		
 		return;
 	}
 
@@ -95,39 +90,6 @@ void FRewindDebuggerTrajectory::Update(float DeltaTime, IRewindDebugger* RewindD
 	
 	// Hook into engine flag for debug drawing.
 	EnsureDebugDrawDelegateExists();
-
-	// Only process trace file after PIE was paused.
-	if (State.bWasPIESimulating)
-	{
-		// Query info from trace file.
-		const TraceServices::IAnalysisSession* Session = RewindDebugger->GetAnalysisSession();
-		if (Session)
-		{
-			State.Extract.TraceStarTime = 0;
-			State.Extract.TraceEndTime = RewindDebugger->CurrentTraceTime();
-		
-			const IGameplayProvider* GameplayProvider = Session->ReadProvider<IGameplayProvider>("GameplayProvider");
-			const IAnimationProvider* AnimationProvider = Session->ReadProvider<IAnimationProvider>("AnimationProvider");
-
-			if (GameplayProvider && AnimationProvider)
-			{
-				TraceServices::FAnalysisSessionReadScope AnalysisSessionReadScope(*Session);
-			
-				BuildTrajectoryOwnersList(RewindDebugger, GameplayProvider, State.Extract.TraceStarTime, State.Extract.TraceEndTime, State.ObjectInfos);
-				BuildTrajectorySkeletalMeshInfoList(GameplayProvider, AnimationProvider, State.Extract.TraceStarTime, State.Extract.TraceEndTime, State.ObjectInfos, State.SkelMeshInfos);
-				BuildTrajectories(GameplayProvider, AnimationProvider, State.Extract.TraceStarTime, State.Extract.TraceEndTime, State.ObjectInfos, State.Trajectories);
-				UpdateDebugInfos(GameplayProvider, State.ObjectInfos, State.DebugInfos);
-			}
-		}
-		else
-		{
-			// Invalid trace session. Abort.
-			Reset();
-			return;
-		}
-		
-		State.bWasPIESimulating = false;
-	}
 	
 	check(State.Trajectories.Num() == State.ObjectInfos.Num());
 }
@@ -149,6 +111,52 @@ void FRewindDebuggerTrajectory::Clear(IRewindDebugger* RewindDebugger)
 	Reset();
 }
 
+void FRewindDebuggerTrajectory::UpdateState(IRewindDebugger* RewindDebugger)
+{
+	// Query info from trace file.
+	const TraceServices::IAnalysisSession* Session = RewindDebugger->GetAnalysisSession();
+	if (Session)
+	{
+		const IGameplayProvider* GameplayProvider = Session->ReadProvider<IGameplayProvider>("GameplayProvider");
+		const IAnimationProvider* AnimationProvider = Session->ReadProvider<IAnimationProvider>("AnimationProvider");
+
+		if (GameplayProvider && AnimationProvider)
+		{
+			TraceServices::FAnalysisSessionReadScope AnalysisSessionReadScope(*Session);
+			
+			// Note that recording index is always zero since RewindDebugger starts a new trace for every recording.
+			if (const IGameplayProvider::RecordingInfoTimeline* Recording = GameplayProvider->GetRecordingInfo(0))
+			{
+				const uint64 EventCount = Recording->GetEventCount();
+
+				if (EventCount > 0)
+				{
+					const FRecordingInfoMessage& LastEvent = Recording->GetEvent(EventCount - 1);
+					
+					State.Extract.TraceStarTime = 0;
+					State.Extract.TraceEndTime = LastEvent.ProfileTime;
+				}
+				else
+				{
+					// Empty trace session. Abort.
+					Reset();
+					return;
+				}
+			}
+			
+			BuildTrajectoryOwnersList(RewindDebugger, GameplayProvider, State.Extract.TraceStarTime, State.Extract.TraceEndTime, State.ObjectInfos);
+			BuildTrajectorySkeletalMeshInfoList(GameplayProvider, AnimationProvider, State.Extract.TraceStarTime, State.Extract.TraceEndTime, State.ObjectInfos, State.SkelMeshInfos);
+			BuildTrajectories(GameplayProvider, AnimationProvider, State.Extract.TraceStarTime, State.Extract.TraceEndTime, State.ObjectInfos, State.Trajectories);
+			UpdateDebugInfos(GameplayProvider, State.ObjectInfos, State.DebugInfos);
+		}
+	}
+	else
+	{
+		// Invalid trace session. Abort.
+		Reset();
+	}
+}
+
 void FRewindDebuggerTrajectory::GenerateMenu()
 {
 	if (UToolMenu* Menu = UToolMenus::Get()->FindMenu("RewindDebugger.ToolBar"))
@@ -276,6 +284,12 @@ void FRewindDebuggerTrajectory::ShowBakeOutTrajectoryWindow()
 	{
 		FSlateApplication::Get().DestroyWindowImmediately(BakeOutWindow.ToSharedRef());
 	}
+
+	// Prepare data before opening window.
+	if (IRewindDebugger* RewindDebugger = IRewindDebugger::Instance())
+	{
+		UpdateState(RewindDebugger);
+	}
 	
 	// Custom window to setup bake out settings and execute the action
 	SAssignNew(BakeOutWindow, SWindow)
diff --git a/Engine/Plugins/Experimental/Animation/TrajectoryTools/Source/TrajectoryTools/Public/TrajectoryRewindDebuggerExtension.h b/Engine/Plugins/Experimental/Animation/TrajectoryTools/Source/TrajectoryTools/Public/TrajectoryRewindDebuggerExtension.h
--- a/Engine/Plugins/Experimental/Animation/TrajectoryTools/Source/TrajectoryTools/Public/TrajectoryRewindDebuggerExtension.h
+++ b/Engine/Plugins/Experimental/Animation/TrajectoryTools/Source/TrajectoryTools/Public/TrajectoryRewindDebuggerExtension.h
@@ -39,6 +39,8 @@ public:
 	
 private:
 
+	void UpdateState(IRewindDebugger* RewindDebugger);
+	
 	// UI
 	
 	void GenerateMenu();
@@ -71,8 +73,6 @@ private:
 		TArray<FObjectInfo> ObjectInfos;			// Associated object's information for the extracted trajectories
 		TArray<FSkeletalMeshInfo> SkelMeshInfos;	// Associated object's skeletal mesh info for the extracted trajectories.	
 		TArray<FDebugInfo> DebugInfos;				// Associated debug information for extracted trajectories
-
-		bool bWasPIESimulating = false;				// Used to request a trajectory extraction operation.
 		
 		/** Reset extension's state */
 		void Reset()
@@ -84,8 +84,6 @@ private:
 			ObjectInfos.Reset();
 			SkelMeshInfos.Reset();
 			DebugInfos.Reset();
-
-			bWasPIESimulating = false;
 		}
 	} State;
 	
