Fixed Trajectory Recording Tool closing after exporting an animation sequence.

#rnx

diff --git a/Engine/Plugins/Experimental/Animation/TrajectoryTools/Source/TrajectoryTools/Private/SExportTrajectoriesWindow.cpp b/Engine/Plugins/Experimental/Animation/TrajectoryTools/Source/TrajectoryTools/Private/SExportTrajectoriesWindow.cpp
--- a/Engine/Plugins/Experimental/Animation/TrajectoryTools/Source/TrajectoryTools/Private/SExportTrajectoriesWindow.cpp
+++ b/Engine/Plugins/Experimental/Animation/TrajectoryTools/Source/TrajectoryTools/Private/SExportTrajectoriesWindow.cpp
@@ -930,12 +930,6 @@ void SExportTrajectoriesWindow::Construct(const FArguments& InArgs)
 						{
 							// Export trajectories to assets.
 							UTrajectoryExportOperation::ExportTrajectory(&Trajectories[SelectedTrajectoryIndex], ExportDetails->ExportSettings, ExportDetails->ExportAssetInfo, TrajectoryOwnerNames[SelectedTrajectoryIndex].ToString());
-
-							// Close window after operation is done.
-							if (WidgetWindow.IsValid())
-							{
-								WidgetWindow.Pin()->RequestDestroyWindow();
-							}
 							
 							return FReply::Handled();
 						})
