Trajectory Recording Tool: Show reason for not being able to bake out sequence. Additionally, instead of opening content browser automatically we offer the option to browse to it via the notification.

#rnx

#Synced-CL 49789214

#Synced-CL 49789214

diff --git a/Engine/Plugins/Experimental/Animation/TrajectoryTools/Source/TrajectoryTools/Private/SExportTrajectoriesWindow.cpp b/Engine/Plugins/Experimental/Animation/TrajectoryTools/Source/TrajectoryTools/Private/SExportTrajectoriesWindow.cpp
--- a/Engine/Plugins/Experimental/Animation/TrajectoryTools/Source/TrajectoryTools/Private/SExportTrajectoriesWindow.cpp
+++ b/Engine/Plugins/Experimental/Animation/TrajectoryTools/Source/TrajectoryTools/Private/SExportTrajectoriesWindow.cpp
@@ -21,6 +21,8 @@
 #include "Framework/Application/SlateApplication.h"
 #include "Engine/Font.h"
 #include "CanvasTypes.h"
+#include "ContentBrowserModule.h"
+#include "IContentBrowserSingleton.h"
 #include "SceneView.h"
 #include "Animation/DebugSkelMeshComponent.h"
 #include "Animation/Skeleton.h"
@@ -909,22 +911,7 @@ void SExportTrajectoriesWindow::Construct(const FArguments& InArgs)
 						.HAlign(HAlign_Center)
 						.IsEnabled_Lambda([this]()
 						{
-							// Early out if export range contains multiple trace ranges.
-							if (PreviewData->OutputTrajectory.TraceInfo.Num() > 1)
-							{
-								int MinIndex = GetSampleIndexFromTime(PreviewData->Details->ExportSettings.FrameRate, PreviewData->Details->ExportSettings.Range.Min, PreviewData->OutputTrajectory.Samples.Num());
-								int MaxIndex = GetSampleIndexFromTime(PreviewData->Details->ExportSettings.FrameRate, PreviewData->Details->ExportSettings.Range.Max, PreviewData->OutputTrajectory.Samples.Num());
-
-								FTrajectoryToolsLibrary::FRangeOverlapTestResult TestResult{};
-								FTrajectoryToolsLibrary::GetRangeOverlaps(PreviewData->OutputTrajectory, TRange<int32>{MinIndex, MaxIndex}, TestResult);
-
-								if (TestResult.bOverlaps && TestResult.Ranges.Num() > 1)
-								{
-									return false;
-								}
-							}
-							
-							return SelectedTrajectoryIndex != INDEX_NONE && ExportDetails && ExportDetails->ExportAssetInfo.IsValid() && ExportDetails->ExportAssetInfo.CanCreateAsset();
+							return GetExportValidation() == EExportValidation::Valid;
 						})
 						.OnClicked_Lambda([this]()
 						{
@@ -933,7 +920,7 @@ void SExportTrajectoriesWindow::Construct(const FArguments& InArgs)
 							
 							return FReply::Handled();
 						})
-						.Text(LOCTEXT("ExportButton", "Bake Out"))
+						.Text_Raw(this, &SExportTrajectoriesWindow::GetButtonLabel)
 					]
 				]
 			]
@@ -1062,6 +1049,75 @@ void SExportTrajectoriesWindow::OnFinishedChangingExportSettingsSelectionPropert
     }
 }
 
+FText SExportTrajectoriesWindow::GetButtonLabel() const
+{
+	const EExportValidation ExportValidationState = GetExportValidation();
+
+	FText ExportValidationResult = FText::GetEmpty();
+	
+	switch (ExportValidationState)
+	{
+	case EExportValidation::MultipleTraceRanges:
+		ExportValidationResult = LOCTEXT("MultipleTraceRangeLabel", "(Multiple trace ranges)");
+		break;
+	case EExportValidation::NoTrajectorySelected:
+		ExportValidationResult = LOCTEXT("NoTrajectorySelectedLabel", "(No trajectory selected)");
+		break;
+	case EExportValidation::MissingAssetInfo:
+		ExportValidationResult = LOCTEXT("MissingAssetInfoLabel", "(Missing asset information)");
+		break;
+	case EExportValidation::ExportPathIsInvalid:
+		ExportValidationResult = LOCTEXT("ExportPathIsInvalidLabel", "(Invalid directory)");
+		break;
+	case EExportValidation::AssetAlreadyExists:
+		ExportValidationResult = LOCTEXT("AssetAlreadyExistsLabel", "(Asset already exists)");
+		break;
+	case EExportValidation::Valid: break;
+	}
+
+	return FText::Format(LOCTEXT("BakeOutLabel", "Bake Out {0}"), ExportValidationResult);
+}
+
+SExportTrajectoriesWindow::EExportValidation SExportTrajectoriesWindow::GetExportValidation() const
+{
+	// Early out if export range contains multiple trace ranges.
+	if (PreviewData->OutputTrajectory.TraceInfo.Num() > 1)
+	{
+		int MinIndex = GetSampleIndexFromTime(PreviewData->Details->ExportSettings.FrameRate, PreviewData->Details->ExportSettings.Range.Min, PreviewData->OutputTrajectory.Samples.Num());
+		int MaxIndex = GetSampleIndexFromTime(PreviewData->Details->ExportSettings.FrameRate, PreviewData->Details->ExportSettings.Range.Max, PreviewData->OutputTrajectory.Samples.Num());
+
+		FTrajectoryToolsLibrary::FRangeOverlapTestResult TestResult{};
+		FTrajectoryToolsLibrary::GetRangeOverlaps(PreviewData->OutputTrajectory, TRange<int32>{MinIndex, MaxIndex}, TestResult);
+
+		if (TestResult.bOverlaps && TestResult.Ranges.Num() > 1)
+		{
+			return EExportValidation::MultipleTraceRanges;
+		}
+	}
+
+	if (SelectedTrajectoryIndex == INDEX_NONE)
+	{
+		return EExportValidation::NoTrajectorySelected;
+	}
+
+	if (ExportDetails == nullptr || !ExportDetails->ExportAssetInfo.IsValid())
+	{
+		return EExportValidation::MissingAssetInfo;	
+	}
+
+	if (!ExportDetails->ExportAssetInfo.IsExportPathValid())
+	{
+		return EExportValidation::ExportPathIsInvalid;
+	}
+
+	if (!ExportDetails->ExportSettings.bShouldOverwriteExistingFiles && ExportDetails->ExportAssetInfo.DoesAssetExistAtPath())
+	{
+		return EExportValidation::AssetAlreadyExists;
+	}
+	
+	return EExportValidation::Valid;
+}
+
 void SExportTrajectoriesWindow::UpdateAssetInfo(FSoftObjectPath&  OutSkeletonPath, FSoftObjectPath & OutSkeletalMeshPath) const
 {
 	if (!ImmutableState.SkeletalMeshInfos.IsValidIndex(SelectedTrajectoryIndex) || !IRewindDebugger::Instance() || !IRewindDebugger::Instance()->GetAnalysisSession())
diff --git a/Engine/Plugins/Experimental/Animation/TrajectoryTools/Source/TrajectoryTools/Private/TrajectoryExportOperation.cpp b/Engine/Plugins/Experimental/Animation/TrajectoryTools/Source/TrajectoryTools/Private/TrajectoryExportOperation.cpp
--- a/Engine/Plugins/Experimental/Animation/TrajectoryTools/Source/TrajectoryTools/Private/TrajectoryExportOperation.cpp
+++ b/Engine/Plugins/Experimental/Animation/TrajectoryTools/Source/TrajectoryTools/Private/TrajectoryExportOperation.cpp
@@ -46,36 +46,50 @@ void FTrajectoryExportSettings::Reset()
 	bShouldExportOnlyAnimatedBones = true;
 }
 
-bool FTrajectoryExportAssetInfo::CanCreateAsset() const
+bool FTrajectoryExportAssetInfo::IsExportPathValid() const
 {
 	FString ObjectPath = FolderPath.Path / AssetName;
-	const FString GameDir = TEXT("/Game/");
-	const FString EngineDir = TEXT("/Engine/");
+	
+	static const FString GameDir = TEXT("/Game/");
+	static const FString EngineDir = TEXT("/Engine/");
+	static const FString PluginsDir = TEXT("/Plugins/");
 
 	bool bDirectoryExists = true;
 	
 	if (ObjectPath.StartsWith(GameDir))
 	{
 		FString ProjectDir = FPaths::ProjectDir();
-		const FString RelativeAssetPath = FPaths::Combine(FPaths::ProjectContentDir(), ObjectPath.Replace(*GameDir, TEXT("")));
+		const FString RelativeAssetFilename = FPaths::Combine(FPaths::ProjectContentDir(), ObjectPath.Replace(*GameDir, TEXT("")));
+		const FString RelativeAssetPath = FPaths::GetPath(RelativeAssetFilename);
+		
 		bDirectoryExists = FPaths::DirectoryExists(FPaths::ConvertRelativePathToFull(RelativeAssetPath));
 	}
 	else if (ObjectPath.StartsWith(EngineDir))
 	{
 		FString ProjectDir = FPaths::ProjectDir();
-		const FString RelativeAssetPath = FPaths::Combine(FPaths::EngineContentDir(), ObjectPath.Replace(*EngineDir, TEXT("")));
+		const FString RelativeAssetFilename = FPaths::Combine(FPaths::EngineContentDir(), ObjectPath.Replace(*EngineDir, TEXT("")));
+		const FString RelativeAssetPath = FPaths::GetPath(RelativeAssetFilename);
+		
 		bDirectoryExists = FPaths::DirectoryExists(FPaths::ConvertRelativePathToFull(RelativeAssetPath));
 	}
 	
-	if (bDirectoryExists)
-	{
-		return false;	
-	}
+	return bDirectoryExists;
+}
+
+bool FTrajectoryExportAssetInfo::DoesAssetExistAtPath() const
+{
+	FString ObjectPath = FolderPath.Path / AssetName;
 	
 	IAssetRegistry& AssetRegistry = FModuleManager::LoadModuleChecked<FAssetRegistryModule>(TEXT("AssetRegistry")).Get();
 	const FAssetData AssetData = AssetRegistry.GetAssetByObjectPath(FSoftObjectPath(ObjectPath));
+	const bool bAssetExists = AssetData.IsValid();
 	
-	return !AssetData.IsValid();
+	return bAssetExists;
+}
+
+bool FTrajectoryExportAssetInfo::CanCreateAsset() const
+{
+	return IsExportPathValid() && !DoesAssetExistAtPath();
 }
 
 void FTrajectoryExportAssetInfo::Reset()
@@ -454,11 +468,6 @@ void UTrajectoryExportOperation::NotifyUserOfResults(const FTrajectoryExportCont
 	{
 		return;
 	}
-	
-	// Select all new assets and show them in the content browser.
-	const FContentBrowserModule& ContentBrowserModule = FModuleManager::Get().LoadModuleChecked<FContentBrowserModule>("ContentBrowser");
-	TArray<FAssetData> NewAssetsData = { AssetToProcess.Get() };
-	ContentBrowserModule.Get().SyncBrowserToAssets(NewAssetsData);
 
 	// Create pop-up notification in editor UI
 	constexpr float NotificationDuration = 5.0f;
@@ -470,7 +479,9 @@ void UTrajectoryExportOperation::NotifyUserOfResults(const FTrajectoryExportCont
 		FNotificationInfo Notification(FText::GetEmpty());
 		Notification.ExpireDuration = NotificationDuration;
 		Notification.Text = FText(LOCTEXT("CancelledExportOperation_NotificationTitle", "Export trajectory cancelled."));
-		FSlateNotificationManager::Get().AddNotification(Notification);
+		
+		TSharedPtr<SNotificationItem> NotificationWidget = FSlateNotificationManager::Get().AddNotification(Notification);
+		NotificationWidget->SetCompletionState(SNotificationItem::CS_Fail);
 	}
 	else
 	{
@@ -489,12 +500,24 @@ void UTrajectoryExportOperation::NotifyUserOfResults(const FTrajectoryExportCont
 		// Notify user that export was completed
 		FNotificationInfo Notification(FText::GetEmpty());
 		Notification.ExpireDuration = NotificationDuration;
+		Notification.ExpireDuration = 10.f;
+		Notification.Hyperlink = FSimpleDelegate::CreateLambda([Asset = GeneratedAsset]()
+		{
+			// Select all new assets and show them in the content browser.
+			const FContentBrowserModule& ContentBrowserModule = FModuleManager::Get().LoadModuleChecked<FContentBrowserModule>("ContentBrowser");
+			TArray<FAssetData> NewAssetsData = { Asset.Get() };
+			ContentBrowserModule.Get().SyncBrowserToAssets(NewAssetsData);
+		});
+
+		Notification.HyperlinkText = LOCTEXT("OpenGeneratedAssetFromTrajectory", "Browse to asset");
 		Notification.Text = FText::Format(
-			LOCTEXT("CompletedExportingTrajectoryNotification", "{0}'s trajectory data was exported to {1}. See Output for details."),
+			LOCTEXT("CompletedExportingTrajectoryNotification", "{0}'s trajectory data was exported to {1} animation sequence."),
 			FText::FromString(Context.SourceObjectName),
 			FText::FromString(*AssetToProcess->GetName()));
-		FSlateNotificationManager::Get().AddNotification(Notification);
-	}
+		
+		TSharedPtr<SNotificationItem> NotificationWidget = FSlateNotificationManager::Get().AddNotification(Notification);
+		NotificationWidget->SetCompletionState(SNotificationItem::CS_Success);
+	}	
 }
 
 void UTrajectoryExportOperation::CleanupIfCancelled(const FScopedSlowTask& Progress) const
diff --git a/Engine/Plugins/Experimental/Animation/TrajectoryTools/Source/TrajectoryTools/Public/SExportTrajectoriesWindow.h b/Engine/Plugins/Experimental/Animation/TrajectoryTools/Source/TrajectoryTools/Public/SExportTrajectoriesWindow.h
--- a/Engine/Plugins/Experimental/Animation/TrajectoryTools/Source/TrajectoryTools/Public/SExportTrajectoriesWindow.h
+++ b/Engine/Plugins/Experimental/Animation/TrajectoryTools/Source/TrajectoryTools/Public/SExportTrajectoriesWindow.h
@@ -184,7 +184,17 @@ public:
 	}
 	
 private:
-	
+
+	enum class EExportValidation
+	{
+		MultipleTraceRanges,
+		NoTrajectorySelected,
+		MissingAssetInfo,
+		ExportPathIsInvalid,
+		AssetAlreadyExists,
+		Valid,
+	};
+
 	// CALLBACKS
 	
 	TSharedRef<SWidget> OnGetTrajectoryPickerMenuContent();
@@ -193,6 +203,8 @@ private:
 
 	// HELPER
 
+	FText GetButtonLabel() const;
+	EExportValidation GetExportValidation() const;
 	void UpdateAssetInfo(FSoftObjectPath&  OutSkeletonPath, FSoftObjectPath & OutSkeletalMeshPath) const;
 	
 	/** Window owning this widget */
diff --git a/Engine/Plugins/Experimental/Animation/TrajectoryTools/Source/TrajectoryTools/Public/TrajectoryExportOperation.h b/Engine/Plugins/Experimental/Animation/TrajectoryTools/Source/TrajectoryTools/Public/TrajectoryExportOperation.h
--- a/Engine/Plugins/Experimental/Animation/TrajectoryTools/Source/TrajectoryTools/Public/TrajectoryExportOperation.h
+++ b/Engine/Plugins/Experimental/Animation/TrajectoryTools/Source/TrajectoryTools/Public/TrajectoryExportOperation.h
@@ -76,6 +76,8 @@ struct FTrajectoryExportAssetInfo
 	UPROPERTY(VisibleAnywhere, Category = "Information")
 	FSoftObjectPath SkeletalMesh;
 
+	UE_API bool IsExportPathValid() const;
+	UE_API bool DoesAssetExistAtPath() const;
 	UE_API bool CanCreateAsset() const;
 	
 	UE_API void Reset();
