Trajectory Recording Tool Improvements:
- Automatically select first valid skeletal mesh to export data from
- Improved UI
- Revamped toolbar buttons for extension.
- Added warning message when we aren't able to export a sequence / trajectory.
- Added world context info for export source, useful to distinguis between client vs server for example.
- Updated labels.
- Miscellaneous improvements / fixes.

#Synced-CL 50651103

diff --git a/Engine/Plugins/Experimental/Animation/TrajectoryTools/Source/TrajectoryTools/Private/SExportTrajectoriesWindow.cpp b/Engine/Plugins/Experimental/Animation/TrajectoryTools/Source/TrajectoryTools/Private/SExportTrajectoriesWindow.cpp
--- a/Engine/Plugins/Experimental/Animation/TrajectoryTools/Source/TrajectoryTools/Private/SExportTrajectoriesWindow.cpp
+++ b/Engine/Plugins/Experimental/Animation/TrajectoryTools/Source/TrajectoryTools/Private/SExportTrajectoriesWindow.cpp
@@ -24,6 +24,7 @@
 #include "ContentBrowserModule.h"
 #include "IContentBrowserSingleton.h"
 #include "SceneView.h"
+#include "SWarningOrErrorBox.h"
 #include "Animation/DebugSkelMeshComponent.h"
 #include "Animation/Skeleton.h"
 #include "TrajectoryExportOperation.h"
@@ -855,10 +856,11 @@ void SExportTrajectoriesWindow::Construct(const FArguments& InArgs)
 							.AutoWidth()
 							[
 								SNew(SRichTextBlock)
-								.Text(LOCTEXT("TrajectyPickerLabel", "Trajectory to export"))
+								.Text(LOCTEXT("TrajectoryPickerLabel", "Source"))
 								.TransformPolicy(ETextTransformPolicy::ToUpper)
 								.DecoratorStyleSet(&FAppStyle::Get())
 								.TextStyle(FAppStyle::Get(), "DetailsView.CategoryTextStyle")
+								.ToolTipText(LOCTEXT("TrajectoryPickerTooltip", "Traced Skeletal Mesh Component to use as the export source for a sequence or trajectory."))
 							]
 
 							// Trajectory picker widget
@@ -894,6 +896,17 @@ void SExportTrajectoriesWindow::Construct(const FArguments& InArgs)
 					]
 				]
 
+				// Warning
+				+SVerticalBox::Slot()
+				.Padding(2)
+				.AutoHeight()
+				[
+					SNew(SWarningOrErrorBox)
+					.Visibility_Lambda([this]{ return GetExportValidation() != EExportValidation::Valid ? EVisibility::Visible : EVisibility::Collapsed; })
+					.MessageStyle(EMessageStyle::Warning)
+					.Message_Lambda([this] { return GetWarningMessage(); })
+				]
+				
 				// Export
 				+ SVerticalBox::Slot()
 				.VAlign(VAlign_Bottom)
@@ -920,13 +933,34 @@ void SExportTrajectoriesWindow::Construct(const FArguments& InArgs)
 							
 							return FReply::Handled();
 						})
-						.Text_Raw(this, &SExportTrajectoriesWindow::GetButtonLabel)
+						.Text(LOCTEXT("ExportLabel", "Export"))
 					]
 				]
 			]
 			
 		]
 	];
+	
+	// Now that everything is constructed properly, default to first trajectory as selected, if we have any.
+	if (!Trajectories.IsEmpty())
+	{
+		// Select first valid trajectory
+		for (int i = 0; i < Trajectories.Num(); ++i)
+		{
+			const bool bHasAnyData = !Trajectories[i].Samples.IsEmpty() && !Trajectories[i].Poses.IsEmpty();
+			
+			if (bHasAnyData)
+			{
+				SelectedTrajectoryIndex = i;
+				break;
+			}
+		}
+		
+		if (SelectedTrajectoryIndex != INDEX_NONE)
+		{
+			UpdateViewportAfterTrajectorySelection();
+		}
+	}
 }
 
 void SExportTrajectoriesWindow::AddReferencedObjects(FReferenceCollector& Collector)
@@ -934,6 +968,46 @@ void SExportTrajectoriesWindow::AddReferencedObjects(FReferenceCollector& Collec
 	Collector.AddReferencedObject(ExportDetails);
 }
 
+void SExportTrajectoriesWindow::UpdateViewportAfterTrajectorySelection()
+{
+	// Reset export settings
+	ExportDetails->Reset();
+	PreviewData->Details = ExportDetails;
+
+	// Update which trajectory we are selecting based on the index
+	PreviewData->SourceIndex = SelectedTrajectoryIndex;
+	PreviewData->Viewport.SourceSkeletalMeshInfo = ImmutableState.SkeletalMeshInfos[SelectedTrajectoryIndex];
+	PreviewData->Viewport.SourceSkeletalMeshComponentId = ImmutableState.DebugInfos[SelectedTrajectoryIndex].OwnerId;
+
+	// Allow to view raw data.
+	PreviewData->SourceTrajectories = MakeConstArrayView(Trajectories);
+				
+	// Match export frame rate
+	PreviewData->OutputTrajectory.Samples.Reset();
+	PreviewData->OutputTrajectory.Poses.Reset();
+	PreviewData->OutputTrajectory.TraceInfo.Reset();
+	FTrajectoryToolsLibrary::TransformTrajectoryToMatchFrameRate(Trajectories[SelectedTrajectoryIndex], ExportDetails->ExportSettings.FrameRate, PreviewData->OutputTrajectory);
+
+	// Always start with entire trajectory range marked for available for export
+	if (!PreviewData->OutputTrajectory.Samples.IsEmpty())
+	{
+		ExportDetails->ExportSettings.Range.Min = PreviewData->OutputTrajectory.Samples[0].Time;
+		ExportDetails->ExportSettings.Range.Max = PreviewData->OutputTrajectory.Samples.Last().Time;
+
+		ExportDetails->NumberOfKeyFrames = PreviewData->OutputTrajectory.Samples.Num();
+		ExportDetails->PlayLength = ExportDetails->ExportSettings.Range.Size();
+	}
+
+	// Update asset info
+	UpdateAssetInfo(ExportDetails->ExportAssetInfo.Skeleton, ExportDetails->ExportAssetInfo.SkeletalMesh);
+				
+	// Focus on preview trajectory
+	PreviewViewportFocusOnTrajectory(PreviewData, Viewport->GetViewportClient());
+
+	// Reset viewport.
+	Viewport->ResetPreviewSkeletalMesh();
+}
+
 TSharedRef<SWidget> SExportTrajectoriesWindow::OnGetTrajectoryPickerMenuContent()
 {
 	const bool bInShouldCloseWindowAfterMenuSelection = true;
@@ -944,7 +1018,7 @@ TSharedRef<SWidget> SExportTrajectoriesWindow::OnGetTrajectoryPickerMenuContent(
 	for (int i = 0; i < TrajectoryOwnerNames.Num(); ++i)
 	{
 		const bool bIsValidTrajectory = Trajectories[i].IsValid();
-		const FText ToolTip = bIsValidTrajectory ? FText::FromName(TrajectoryOwnerNames[i]) : LOCTEXT("InvalidTrajectoryTooltip", "Invalid. Trajectory data is empty or there is a mistmatch between poses and samples.");
+		const FText ToolTip = bIsValidTrajectory ? FText::FromName(TrajectoryOwnerNames[i]) : LOCTEXT("InvalidTrajectoryTooltip", "Invalid. Trajectory data is empty or there is a mismatch between poses and samples.");
 
 		FUIAction UIAction;
 
@@ -961,42 +1035,7 @@ TSharedRef<SWidget> SExportTrajectoriesWindow::OnGetTrajectoryPickerMenuContent(
 
 			if (PrevSelectedTrajectoryIndex != SelectedTrajectoryIndex && SelectedTrajectoryIndex != INDEX_NONE)
 			{
-				// Reset export settings
-				ExportDetails->Reset();
-				PreviewData->Details = ExportDetails;
-
-				// Update which trajectory we are selecting based on the index
-				PreviewData->SourceIndex = SelectedTrajectoryIndex;
-				PreviewData->Viewport.SourceSkeletalMeshInfo = ImmutableState.SkeletalMeshInfos[SelectedTrajectoryIndex];
-				PreviewData->Viewport.SourceSkeletalMeshComponentId = ImmutableState.DebugInfos[SelectedTrajectoryIndex].OwnerId;
-
-				// Allow to view raw data.
-				PreviewData->SourceTrajectories = MakeConstArrayView(Trajectories);
-				
-				// Match export frame rate
-				PreviewData->OutputTrajectory.Samples.Reset();
-				PreviewData->OutputTrajectory.Poses.Reset();
-				PreviewData->OutputTrajectory.TraceInfo.Reset();
-				FTrajectoryToolsLibrary::TransformTrajectoryToMatchFrameRate(Trajectories[SelectedTrajectoryIndex], ExportDetails->ExportSettings.FrameRate, PreviewData->OutputTrajectory);
-
-				// Always start with entire trajectory range marked for available for export
-				if (!PreviewData->OutputTrajectory.Samples.IsEmpty())
-				{
-					ExportDetails->ExportSettings.Range.Min = PreviewData->OutputTrajectory.Samples[0].Time;
-					ExportDetails->ExportSettings.Range.Max = PreviewData->OutputTrajectory.Samples.Last().Time;
-
-					ExportDetails->NumberOfKeyFrames = PreviewData->OutputTrajectory.Samples.Num();
-					ExportDetails->PlayLength = ExportDetails->ExportSettings.Range.Size();
-				}
-
-				// Update asset info
-				UpdateAssetInfo(ExportDetails->ExportAssetInfo.Skeleton, ExportDetails->ExportAssetInfo.SkeletalMesh);
-				
-				// Focus on preview trajectory
-				PreviewViewportFocusOnTrajectory(PreviewData, Viewport->GetViewportClient());
-
-				// Reset viewport.
-				Viewport->ResetPreviewSkeletalMesh();
+				UpdateViewportAfterTrajectorySelection();
 			}
 		});
 		
@@ -1049,33 +1088,21 @@ void SExportTrajectoriesWindow::OnFinishedChangingExportSettingsSelectionPropert
     }
 }
 
-FText SExportTrajectoriesWindow::GetButtonLabel() const
+FText SExportTrajectoriesWindow::GetWarningMessage() const
 {
 	const EExportValidation ExportValidationState = GetExportValidation();
-
-	FText ExportValidationResult = FText::GetEmpty();
 	
 	switch (ExportValidationState)
 	{
-	case EExportValidation::MultipleTraceRanges:
-		ExportValidationResult = LOCTEXT("MultipleTraceRangeLabel", "(Multiple trace ranges)");
-		break;
-	case EExportValidation::NoTrajectorySelected:
-		ExportValidationResult = LOCTEXT("NoTrajectorySelectedLabel", "(No trajectory selected)");
-		break;
-	case EExportValidation::MissingAssetInfo:
-		ExportValidationResult = LOCTEXT("MissingAssetInfoLabel", "(Missing asset information)");
-		break;
-	case EExportValidation::ExportPathIsInvalid:
-		ExportValidationResult = LOCTEXT("ExportPathIsInvalidLabel", "(Invalid directory)");
-		break;
-	case EExportValidation::AssetAlreadyExists:
-		ExportValidationResult = LOCTEXT("AssetAlreadyExistsLabel", "(Asset already exists)");
-		break;
-	case EExportValidation::Valid: break;
+	case EExportValidation::MultipleTraceRanges: return LOCTEXT("MultipleTraceRangeLabel", "Multiple trace ranges");
+	case EExportValidation::NoTrajectorySelected: return  LOCTEXT("NoTrajectorySelectedLabel", "No trajectory selected");
+	case EExportValidation::MissingAssetInfo: return LOCTEXT("MissingAssetInfoLabel", "Missing asset information");
+	case EExportValidation::ExportPathIsInvalid: return LOCTEXT("ExportPathIsInvalidLabel", "Invalid directory");
+	case EExportValidation::AssetAlreadyExists: return LOCTEXT("AssetAlreadyExistsLabel", "Asset already exists");
+	case EExportValidation::Valid:
+	default:
+		return FText::GetEmpty();
 	}
-
-	return FText::Format(LOCTEXT("BakeOutLabel", "Bake Out {0}"), ExportValidationResult);
 }
 
 SExportTrajectoriesWindow::EExportValidation SExportTrajectoriesWindow::GetExportValidation() const
diff --git a/Engine/Plugins/Experimental/Animation/TrajectoryTools/Source/TrajectoryTools/Private/TrajectoryExportOperation.cpp b/Engine/Plugins/Experimental/Animation/TrajectoryTools/Source/TrajectoryTools/Private/TrajectoryExportOperation.cpp
--- a/Engine/Plugins/Experimental/Animation/TrajectoryTools/Source/TrajectoryTools/Private/TrajectoryExportOperation.cpp
+++ b/Engine/Plugins/Experimental/Animation/TrajectoryTools/Source/TrajectoryTools/Private/TrajectoryExportOperation.cpp
@@ -95,7 +95,7 @@ bool FTrajectoryExportAssetInfo::CanCreateAsset() const
 void FTrajectoryExportAssetInfo::Reset()
 {
 	AssetName = "";
-	FolderPath.Path = TEXT("/Game");
+	FolderPath.Path = TEXT("");
 	Skeleton.Reset();
 	SkeletalMesh.Reset();
 }
diff --git a/Engine/Plugins/Experimental/Animation/TrajectoryTools/Source/TrajectoryTools/Private/TrajectoryRewindDebuggerExtension.cpp b/Engine/Plugins/Experimental/Animation/TrajectoryTools/Source/TrajectoryTools/Private/TrajectoryRewindDebuggerExtension.cpp
--- a/Engine/Plugins/Experimental/Animation/TrajectoryTools/Source/TrajectoryTools/Private/TrajectoryRewindDebuggerExtension.cpp
+++ b/Engine/Plugins/Experimental/Animation/TrajectoryTools/Source/TrajectoryTools/Private/TrajectoryRewindDebuggerExtension.cpp
@@ -19,6 +19,9 @@
 #include "Animation/AnimSequence.h"
 #include "Widgets/Input/SButton.h"
 #include "Logging/StructuredLog.h"
+#include "Styling/SlateIconFinder.h"
+#include "Styling/ToolBarStyle.h"
+#include "Widgets/Images/SImage.h"
 
 #define LOCTEXT_NAMESPACE "RewindDebuggerTrajectory"
 
@@ -74,7 +77,7 @@ void FRewindDebuggerTrajectory::Reset()
 void FRewindDebuggerTrajectory::Update(float DeltaTime, IRewindDebugger* RewindDebugger)
 {
 	check(RewindDebugger);
-	
+
 	// Early out.
 	if (RewindDebugger->IsPIESimulating() || RewindDebugger->GetRecordingDuration() == 0)
 	{
@@ -168,67 +171,95 @@ void FRewindDebuggerTrajectory::GenerateMenu()
 
 	if (UToolMenu* Menu = UToolMenus::Get()->FindMenu(RewindDebuggerModule->GetMainToolbarName()))
 	{
-		FToolMenuSection& Section = Menu->AddSection("TrajectoryWorkflows", LOCTEXT("Trajectory Tools","Trajectory Tools"));
+		FToolMenuSection& Section = Menu->FindOrAddSection("TrajectoryWorkflows", LOCTEXT("Trajectory Tools","Trajectory Tools"));
 		
 		Section.AddSeparator("TrajectoryWorkflows");
 		{
-			static const FName Name = "Trajectories";
-			static const FText Label = LOCTEXT("TrajectoriesMenuLabel", "Trajectories");
-			static const FText ToolTip = LOCTEXT("TrajectoriesMenuTooltip", "Toggle trajectories to display in viewport");
+			// Export trajectories
+			{
+				static const FName Name = "ExportToolbarMenu"; 
+				static const FText Label = LOCTEXT("ExportTrajectoriesLabel", "Export Recording");
+				static const FText Tooltip = LOCTEXT("ExportTrajectoriesTooltip", "Export recording/trajectory into a standalone asset");
+				
+				const FToolBarStyle& ToolBarStyle = FAppStyle::Get().GetWidgetStyle<FToolBarStyle>("SlimToolBar");
+				
+				TSharedRef<SWidget> CustomButton = SNew(SButton)
+				.ButtonStyle(&ToolBarStyle.ButtonStyle)
+				.ToolTipText(Tooltip)
+				.IsEnabled_Lambda([]()
+				{
+					IRewindDebugger* RewindDebugger = IRewindDebugger::Instance();
+					return RewindDebugger ? !RewindDebugger->IsPIESimulating() || !RewindDebugger->IsRecording() : false;
+				})
+				.OnClicked_Lambda([this]()
+				{
+					ShowBakeOutTrajectoryWindow();
+					return FReply::Handled();
+				})
+				[
+					SNew(SHorizontalBox)
+
+					+ SHorizontalBox::Slot()
+					.AutoWidth()
+					.VAlign(VAlign_Center)
+					.Padding(0.f, 0.f, 6.f, 0.f)
+					[
+						SNew(SImage)
+						.Image(FAppStyle::GetBrush("Icons.Toolbar.Export"))
+					]
+
+					+ SHorizontalBox::Slot()
+					.AutoWidth()
+					.VAlign(VAlign_Center)
+					[
+						SNew(STextBlock)
+						.Text(Label)
+						.TextStyle(&ToolBarStyle.LabelStyle)
+					]
+				];
+				
+				Section.AddEntry(FToolMenuEntry::InitWidget(Name, CustomButton, Label, false, true, false, Tooltip));
+			}
 			
-			Section.AddSubMenu(Name, Label, ToolTip, FNewToolMenuDelegate::CreateLambda([this](UToolMenu* InMenu)
+			// Preview trajectories drawing
 			{
-				FToolMenuSection& Section = InMenu->AddSection(NAME_None);
-				Section.AddDynamicEntry("Trajectories", FNewToolMenuSectionDelegate::CreateRaw(this, &FRewindDebuggerTrajectory::MakeTrajectoriesMenu));
-			}));
+				static const FName Name = "ToggleDebugDrawComboButton";
+				static const FText Label = LOCTEXT("ToggleDebugDrawTrajectoriesLabel", "Preview trajectories");
+				static const FText ToolTip = LOCTEXT("ToggleDebugDrawTrajectoriesTooltip", "Toggle recorded trajectories to display in viewport");
+				
+				Section.AddEntry(FToolMenuEntry::InitSubMenu(Name, Label, ToolTip, FNewToolMenuDelegate::CreateRaw(this, &FRewindDebuggerTrajectory::MakeToggleDebugDrawMenu), false, FSlateIcon(), false));
+			}
 		}
 	}
 }
 
-void FRewindDebuggerTrajectory::MakeTrajectoriesMenu(FToolMenuSection& InSection)
+void FRewindDebuggerTrajectory::MakeToggleDebugDrawMenu(UToolMenu* InMenu)
 {
-	// Toggle trajectory drawing
-	{
-		static const FName Name = "ToggleDebugDrawComboButton";
-		static const FText Label = LOCTEXT("ToggleDebugDrawTrajectoriesLabel", "Toggle debug draw");
-		static const FText ToolTip = LOCTEXT("ToggleDebugDrawTrajectoriesTooltip", "Toggle trajectories to display in viewport");
-
-		InSection.AddSubMenu(Name, Label, ToolTip, FNewToolMenuDelegate::CreateRaw(this, &FRewindDebuggerTrajectory::MakeToggleDebugDrawMenu), false, FSlateIcon(), false);
-	}
-
-	// Bake out trajectories
+	// Prepare trajectories to preview.
+	if (IRewindDebugger* RewindDebugger = IRewindDebugger::Instance())
 	{
-		FUIAction UIAction;
-		UIAction.ExecuteAction = FExecuteAction::CreateLambda([this](){ ShowBakeOutTrajectoryWindow(); });
-		UIAction.CanExecuteAction = FCanExecuteAction::CreateLambda([this]()
-		{
-			IRewindDebugger* RewindDebugger = IRewindDebugger::Instance();
-			return RewindDebugger ? !RewindDebugger->IsPIESimulating() : false;
-		});
-		
-		static const FName Name = "BakeOutSubMenu"; 
-		static const FText Label = LOCTEXT("BakeOutTrajectoriesLabel", "Bake out...");
-		static const FText Tooltip = LOCTEXT("BakeOutTrajectoriesTooltip", "Bake a trajectory into a standalone asset");
-		
-		InSection.AddMenuEntry(Name, Label, Tooltip, FSlateIcon(), UIAction);
+		UpdateState(RewindDebugger);
 	}
-}
-
-void FRewindDebuggerTrajectory::MakeToggleDebugDrawMenu(UToolMenu* InMenu)
-{
+	
 	FToolMenuSection& Section = InMenu->FindOrAddSection("ToggleDebugDrawOptions");
 
-	for (FExtensionState::FDebugInfo& DebugInfoEntry : State.DebugInfos)
+	for (int i = 0; i < State.Trajectories.Num(); ++i)
 	{
+		FExtensionState::FDebugInfo& DebugInfoEntry = State.DebugInfos[i];
+		
 		FUIAction UIAction;
 		{
 			// @todo: For safety, this should probably just copy the component id and try to find it, then toggle it.
 			UIAction.ExecuteAction = FExecuteAction::CreateLambda([&DebugInfoEntry](){ DebugInfoEntry.bShouldDraw = !DebugInfoEntry.bShouldDraw; });
 			UIAction.GetActionCheckState = FGetActionCheckState::CreateLambda([&DebugInfoEntry] { return DebugInfoEntry.bShouldDraw ? ECheckBoxState::Checked : ECheckBoxState::Unchecked; });
+			UIAction.CanExecuteAction = FCanExecuteAction::CreateLambda([this, i]()
+				{
+					const bool bHasAnyData = !State.Trajectories[i].Samples.IsEmpty() && !State.Trajectories[i].Poses.IsEmpty();
+					return bHasAnyData;
+				});
 		}
 
-		// @todo: Add an actor icon for the trajectory picker. Similar to RewindDebugger's.
-		Section.AddMenuEntry(DebugInfoEntry.Name, FText::FromName(DebugInfoEntry.Name), FText::FromName(DebugInfoEntry.Name), FSlateIcon(), UIAction, EUserInterfaceActionType::Check);
+		Section.AddMenuEntry(DebugInfoEntry.Name, FText::FromName(DebugInfoEntry.Name), FText::FromName(DebugInfoEntry.Name), FSlateIconFinder::FindIconForClass(USkeletalMeshComponent::StaticClass(), TEXT("SCS.Component")), UIAction, EUserInterfaceActionType::Check);
 	}
 }
 
@@ -255,7 +286,8 @@ void FRewindDebuggerTrajectory::DebugDraw(UCanvas* InCanvas, APlayerController*
 				for (int SampleIndex = 0; SampleIndex < DebugTrajectoryPair.Samples.Num(); ++SampleIndex)
 				{
 					const FGameplayTrajectory::FSample& Sample = DebugTrajectoryPair.Samples[SampleIndex];
-
+					
+					// Draw trajectory
 					if (SampleIndex + 1 < DebugTrajectoryPair.Samples.Num())
 					{
 						const FGameplayTrajectory::FSample& FutureSample = DebugTrajectoryPair.Samples[SampleIndex + 1];
@@ -263,21 +295,10 @@ void FRewindDebuggerTrajectory::DebugDraw(UCanvas* InCanvas, APlayerController*
 						DrawDebugLine(WorldToVisualize, Sample.Position, FutureSample.Position, DebugDrawColorForTrajectory);
 					}
 
-					// TODO: Determine when to draw text. AKA every second, every X frames, etc.
-					if (SampleIndex % 10 == 0)
+					// Draw sections at every second
+					if (SampleIndex % 20 == 0)
 					{
-						FVector2D SamplePositionOnScreen;
-					
-						if (InCanvas->SceneView->WorldToPixel(Sample.Position, SamplePositionOnScreen))
-						{
-							FFontRenderInfo Info;
-							Info.bEnableShadow = true;
-					
-							SamplePositionOnScreen.X = FMath::RoundToFloat(SamplePositionOnScreen.X);
-							SamplePositionOnScreen.Y = FMath::RoundToFloat(SamplePositionOnScreen.Y);
-						
-							InCanvas->DrawText(GEngine->GetTinyFont(), FString::SanitizeFloat(Sample.Time), SamplePositionOnScreen.X, SamplePositionOnScreen.Y, 1, 1, Info);
-						}
+						DrawDebugPoint(WorldToVisualize, Sample.Position, 5, DebugDrawColorForTrajectory);
 					}
 				}
 			}
@@ -300,7 +321,7 @@ void FRewindDebuggerTrajectory::ShowBakeOutTrajectoryWindow()
 	
 	// Custom window to setup bake out settings and execute the action
 	SAssignNew(BakeOutWindow, SWindow)
-	.Title(LOCTEXT("BakeOutTrajectoryWindowTitle", "Bake Out Trajectories"))
+	.Title(LOCTEXT("BakeOutTrajectoryWindowTitle", "Export Recording / Trajectories"))
 	.SizingRule(ESizingRule::UserSized)
 	.SupportsMaximize(true)
 	.SupportsMinimize(true)
@@ -352,7 +373,12 @@ void FRewindDebuggerTrajectory::BuildTrajectoryOwnersList(const IRewindDebugger*
 	
 	// Used to filter which objects we can use to build our trajectory.
 	const FClassInfo* SkelMeshCmpClassInfo = InGameplayProvider->FindClassInfo(*USkeletalMeshComponent::StaticClass()->GetPathName());
-	check(SkelMeshCmpClassInfo)
+	
+	// Not a single actor with a skeletal mesh component was traced, safe to early out.
+	if (SkelMeshCmpClassInfo == nullptr)
+	{
+		return;
+	}
 	
 	// Query target actor of the trace file / session.
 	InGameplayProvider->EnumerateObjects(InTraceStartTime, InTraceEndTime, [&OutObjectInfos,InGameplayProvider, SkelMeshCmpClassInfo, InRewindDebugger](const FObjectInfo& InObjectInfo)
@@ -362,10 +388,11 @@ void FRewindDebuggerTrajectory::BuildTrajectoryOwnersList(const IRewindDebugger*
 		{
 			InGameplayProvider->EnumerateSubobjects(InObjectInfo.GetId(), [&OutObjectInfos, InGameplayProvider, SkelMeshCmpClassInfo](const RewindDebugger::FObjectId& InSubObjectId)
 			{
+				const bool bIsSubObjectValid = InSubObjectId.IsSet();
 				const FObjectInfo& SubObjectInfo = InGameplayProvider->GetObjectInfo(InSubObjectId);
 
 				const bool bIsSubObjectASkelMesh = SubObjectInfo.ClassId == SkelMeshCmpClassInfo->Id || InGameplayProvider->IsSubClassOf(SubObjectInfo.ClassId, SkelMeshCmpClassInfo->Id);
-				if (bIsSubObjectASkelMesh)
+				if (bIsSubObjectValid && bIsSubObjectASkelMesh)
 				{
 					OutObjectInfos.Push(SubObjectInfo);
 				}
@@ -373,8 +400,9 @@ void FRewindDebuggerTrajectory::BuildTrajectoryOwnersList(const IRewindDebugger*
 		}
 		else
 		{
+			const bool bIsObjectValid = InObjectInfo.GetId().IsSet();
 			const bool bIsObjectASkelMesh = InObjectInfo.ClassId == SkelMeshCmpClassInfo->Id || InGameplayProvider->IsSubClassOf(InObjectInfo.ClassId, SkelMeshCmpClassInfo->Id);
-			if (bIsObjectASkelMesh)
+			if (bIsObjectValid && bIsObjectASkelMesh)
 			{
 				const bool bHasOwningActor = InRewindDebugger->FindTypedOuterInfo<AActor>(InGameplayProvider, InObjectInfo.GetUObjectId()) != nullptr;
 				if (bHasOwningActor)
@@ -455,7 +483,12 @@ void FRewindDebuggerTrajectory::BuildTrajectories(const IGameplayProvider* InGam
 
 	// Used to filter which objects we can use to build our trajectory.
 	const FClassInfo* SkelMeshCmpClassInfo = InGameplayProvider->FindClassInfo(*USkeletalMeshComponent::StaticClass()->GetPathName());
-	check(SkelMeshCmpClassInfo)
+
+	// Not a single actor with a skeletal mesh component was traced, safe to early out.
+	if (SkelMeshCmpClassInfo == nullptr)
+	{
+		return;
+	}
 
 	// Build trajectory from pose's root transform over time.
 	for (const FObjectInfo& ObjectInfo : InObjectInfos)
@@ -594,16 +627,18 @@ FName FRewindDebuggerTrajectory::GetFullNameForDebugInfoOwner(const IGameplayPro
 			Result.Prepend(TEXTVIEW(" - "));
 			Result.Prepend(ObjectInfo.Name);
 
+			// Add world context
+			Result.Append(TEXTVIEW(" ("));
 			const FWorldInfo* WorldInfo = InGameplayProvider->FindWorldInfoFromObject(ObjectId);
 			if (WorldInfo)
 			{
-				bool bIsServer = WorldInfo->NetMode == FWorldInfo::ENetMode::DedicatedServer;
-
-				if (bIsServer)
-				{
-					Result.Append(TEXTVIEW(" (Server)"));
-				}
+				Result.Append(GetDebugInfoOwnerWorldPostFix(*WorldInfo).ToString());
 			}
+			else
+			{
+				Result.Append(TEXTVIEW("Unknown"));
+			}
+			Result.Append(TEXTVIEW(")"));
 			
 			return FName(Result.ToView());
 		}
@@ -622,10 +657,69 @@ FName FRewindDebuggerTrajectory::GetFullNameForDebugInfoOwner(const IGameplayPro
 			}
 			else
 			{
-				return NAME_None;
+				// Add world context
+				Result.Append(TEXTVIEW(" ("));
+				const FWorldInfo* WorldInfo = InGameplayProvider->FindWorldInfoFromObject(ObjectId);
+				if (WorldInfo)
+				{
+					Result.Append(GetDebugInfoOwnerWorldPostFix(*WorldInfo).ToString());
+				}
+				else
+				{
+					Result.Append(TEXTVIEW("Unknown"));
+				}
+				Result.Append(TEXTVIEW(")"));
+				
+				return FName(Result.ToView());
 			}
 		}
 	}
 }
 
+// copied from FGameplayTrack::GetWorldPostFix()
+FText FRewindDebuggerTrajectory::GetDebugInfoOwnerWorldPostFix(const FWorldInfo& InWorldInfo)
+{
+	switch (InWorldInfo.Type)
+	{
+		default:
+			return LOCTEXT("Unknown", "Unknown");
+		
+		case FWorldInfo::EType::None:
+			return LOCTEXT("None", "None");
+		
+		case FWorldInfo::EType::PIE:
+			{
+				switch (InWorldInfo.NetMode)
+				{
+				case FWorldInfo::ENetMode::Client:
+					return FText::Format(LOCTEXT("ClientIndexFormat", "Client {0}"), InWorldInfo.PIEInstanceId - 1);
+					
+				case FWorldInfo::ENetMode::DedicatedServer:
+					return LOCTEXT("ServerPostfix", "Server");
+					
+				case FWorldInfo::ENetMode::ListenServer:
+					return LOCTEXT("ListenServerPostfix", "Listen Server");
+					
+				case FWorldInfo::ENetMode::Standalone:
+					return InWorldInfo.bIsSimulating ? LOCTEXT("SimulateInEditorPostfix", "Simulate") : LOCTEXT("PlayInEditorPostfix", "PIE");
+					
+				default: 
+					return LOCTEXT("UnknownNetModePIEPostFix", "PIE");
+				}
+			}
+		
+		case FWorldInfo::EType::Editor:
+		case FWorldInfo::EType::EditorPreview:
+			return LOCTEXT("EditorPostfix", "Editor");
+		
+		case FWorldInfo::EType::Game:
+		case FWorldInfo::EType::GamePreview:
+		case FWorldInfo::EType::GameRPC:
+			return LOCTEXT("GamePostfix", "Game");
+		
+		case FWorldInfo::EType::Inactive:
+			return LOCTEXT("InactivePostfix", "Inactive");
+	}
+}
+
 #undef LOCTEXT_NAMESPACE
diff --git a/Engine/Plugins/Experimental/Animation/TrajectoryTools/Source/TrajectoryTools/Public/SExportTrajectoriesWindow.h b/Engine/Plugins/Experimental/Animation/TrajectoryTools/Source/TrajectoryTools/Public/SExportTrajectoriesWindow.h
--- a/Engine/Plugins/Experimental/Animation/TrajectoryTools/Source/TrajectoryTools/Public/SExportTrajectoriesWindow.h
+++ b/Engine/Plugins/Experimental/Animation/TrajectoryTools/Source/TrajectoryTools/Public/SExportTrajectoriesWindow.h
@@ -177,11 +177,6 @@ public:
 	virtual FString GetReferencerName() const override	{ return TEXT("SExportTrajectoriesWindow"); }
 	virtual void AddReferencedObjects(FReferenceCollector& Collector) override;
 	/** End FGCObject */
-
-	void ResetPreviewSelection()
-	{
-		SelectedTrajectoryIndex = INDEX_NONE;
-	}
 	
 private:
 
@@ -203,9 +198,10 @@ private:
 
 	// HELPER
 
-	FText GetButtonLabel() const;
+	FText GetWarningMessage() const;
 	EExportValidation GetExportValidation() const;
 	void UpdateAssetInfo(FSoftObjectPath&  OutSkeletonPath, FSoftObjectPath & OutSkeletalMeshPath) const;
+	void UpdateViewportAfterTrajectorySelection();
 	
 	/** Window owning this widget */
 	TWeakPtr<SWindow> WidgetWindow;
diff --git a/Engine/Plugins/Experimental/Animation/TrajectoryTools/Source/TrajectoryTools/Public/TrajectoryRewindDebuggerExtension.h b/Engine/Plugins/Experimental/Animation/TrajectoryTools/Source/TrajectoryTools/Public/TrajectoryRewindDebuggerExtension.h
--- a/Engine/Plugins/Experimental/Animation/TrajectoryTools/Source/TrajectoryTools/Public/TrajectoryRewindDebuggerExtension.h
+++ b/Engine/Plugins/Experimental/Animation/TrajectoryTools/Source/TrajectoryTools/Public/TrajectoryRewindDebuggerExtension.h
@@ -48,6 +48,7 @@ private:
 	void MakeToggleDebugDrawMenu(UToolMenu* InMenu);
 	
 	// Debug Draw
+	
 	void EnsureDebugDrawDelegateExists();
 	void DebugDraw(UCanvas* InCanvas, APlayerController* InController);
 
@@ -99,6 +100,7 @@ private:
 	static void BuildTrajectories(const IGameplayProvider* InGameplayProvider, const IAnimationProvider* InAnimationProvider, double InTraceStartTime, double InTraceEndTime, const TArray<FObjectInfo>& InObjectInfos, TArray<FGameplayTrajectory>& OutTrajectories);
 	static void UpdateDebugInfos(const IGameplayProvider* InGameplayProvider, const TArray<FObjectInfo>& InObjectInfos, TArray<FExtensionState::FDebugInfo>& InOutDebugInfos);
 	static FName GetFullNameForDebugInfoOwner(const IGameplayProvider* InGameplayProvider, uint64 InOwnerObjectId);
-
+	static FText GetDebugInfoOwnerWorldPostFix(const FWorldInfo& InWorldInfo);
+	
 	friend class SExportTrajectoriesWindow;
 };
