[Trajectory Tools]
Replace FAnalysisSessionReadScope with granular FProviderReadScopeLock for individual providers in the trajectory export window
Simplify UpdatePreviewPoseFromScrubTime by removing provider parameters now that locking is handled internally

#rnx
#Synced-CL 52460709

diff --git a/Engine/Plugins/Experimental/Animation/TrajectoryTools/Source/TrajectoryTools/Private/SExportTrajectoriesWindow.cpp b/Engine/Plugins/Experimental/Animation/TrajectoryTools/Source/TrajectoryTools/Private/SExportTrajectoriesWindow.cpp
--- a/Engine/Plugins/Experimental/Animation/TrajectoryTools/Source/TrajectoryTools/Private/SExportTrajectoriesWindow.cpp
+++ b/Engine/Plugins/Experimental/Animation/TrajectoryTools/Source/TrajectoryTools/Private/SExportTrajectoriesWindow.cpp
@@ -2,6 +2,7 @@
 
 #include "SExportTrajectoriesWindow.h"
 
+#include "Common/ProviderLock.h"
 #include "PropertyCustomizationHelpers.h"
 #include "SSimpleTimeSlider.h"
 #include "Widgets/Input/SComboButton.h"
@@ -116,12 +117,10 @@ void SPreviewTrajectoryViewport::ResetPreviewSkeletalMesh() const
 	// Preview traced skeletal mesh.
 	if (GameplayProvider && AnimationProvider)
 	{
-		TraceServices::FAnalysisSessionReadScope AnalysisSessionReadScope(*Session);
-
 		FSkeletalMeshInfo SkeletalMeshInfoToPreview {};
 		{
 			const int32 ScrubSampleIndex = GetSampleIndexFromTime(PreviewData->Details->ExportSettings.FrameRate, PreviewData->Timeline.ScrubTime, PreviewData->OutputTrajectory.Samples.Num());
-		
+
 			for (int i = 0; i < PreviewData->OutputTrajectory.TraceInfo.Num(); ++i)
 			{
 				if (PreviewData->OutputTrajectory.TraceInfo.Ranges[i].Contains(ScrubSampleIndex))
@@ -132,43 +131,47 @@ void SPreviewTrajectoryViewport::ResetPreviewSkeletalMesh() const
 			}
 		}
 		ensureMsgf(SkeletalMeshInfoToPreview.Id != 0, TEXT("Invalid skeletal mesh to preview. If this gets hit something went wrong."));
-		
-		if (const FSkeletalMeshInfo* SkeletalMeshInfo = AnimationProvider->FindSkeletalMeshInfo(SkeletalMeshInfoToPreview.Id /*PreviewData->Viewport.SourceSkeletalMeshInfo.Id*/))
+
 		{
-			if (const FObjectInfo* SkeletalMeshObjectInfo = GameplayProvider->FindObjectInfo(SkeletalMeshInfoToPreview.Id /*PreviewData->Viewport.SourceSkeletalMeshInfo.Id)*/))
+			TraceServices::FProviderReadScopeLock AnimationProviderReadScope(*AnimationProvider);
+			if (const FSkeletalMeshInfo* SkeletalMeshInfo = AnimationProvider->FindSkeletalMeshInfo(SkeletalMeshInfoToPreview.Id /*PreviewData->Viewport.SourceSkeletalMeshInfo.Id*/))
 			{
-				USkeletalMesh* SkeletalMesh = TSoftObjectPtr<USkeletalMesh>(FSoftObjectPath(SkeletalMeshObjectInfo->PathName)).LoadSynchronous();
-
-				// Try the skeleton asset preview mesh as a fallback when skeletal mesh asset is not valid.
-				if (SkeletalMesh == nullptr)
+				TraceServices::FProviderReadScopeLock GameplayProviderReadScope(*GameplayProvider);
+				if (const FObjectInfo* SkeletalMeshObjectInfo = GameplayProvider->FindObjectInfo(SkeletalMeshInfoToPreview.Id /*PreviewData->Viewport.SourceSkeletalMeshInfo.Id)*/))
 				{
-					// Ensure our skeleton is valid.
-					if (SkeletalMeshInfo->SkeletonId != 0)
+					USkeletalMesh* SkeletalMesh = TSoftObjectPtr<USkeletalMesh>(FSoftObjectPath(SkeletalMeshObjectInfo->PathName)).LoadSynchronous();
+
+					// Try the skeleton asset preview mesh as a fallback when skeletal mesh asset is not valid.
+					if (SkeletalMesh == nullptr)
 					{
-						const FObjectInfo& SkeletonObjectInfo = GameplayProvider->GetObjectInfo(SkeletalMeshInfo->SkeletonId);
-					
-						USkeleton* Skeleton = TSoftObjectPtr<USkeleton>(FSoftObjectPath(SkeletonObjectInfo.PathName)).LoadSynchronous();
-						if (Skeleton)
+						// Ensure our skeleton is valid.
+						if (SkeletalMeshInfo->SkeletonId != 0)
 						{
-							SkeletalMesh = Skeleton->GetPreviewMesh(true);
+							const FObjectInfo& SkeletonObjectInfo = GameplayProvider->GetObjectInfo(SkeletalMeshInfo->SkeletonId);
+					
+							USkeleton* Skeleton = TSoftObjectPtr<USkeleton>(FSoftObjectPath(SkeletonObjectInfo.PathName)).LoadSynchronous();
+							if (Skeleton)
+							{
+								SkeletalMesh = Skeleton->GetPreviewMesh(true);
+							}
 						}
 					}
-				}
 
-				// Assign loaded skeletal mesh.
-				if (SkeletalMesh)
-				{
-					PreviewComponent->SetSkeletalMesh(SkeletalMesh);
+					// Assign loaded skeletal mesh.
+					if (SkeletalMesh)
+					{
+						PreviewComponent->SetSkeletalMesh(SkeletalMesh);
+					}
 				}
 			}
 		}
 
 		// Run method that sets world position and bone transforms.
-		UpdatePreviewPoseFromScrubTime(AnimationProvider, GameplayProvider, 0, RewindDebugger->CurrentTraceTime());
+		UpdatePreviewPoseFromScrubTime(0, RewindDebugger->CurrentTraceTime());
 	}
 }
 
-void SPreviewTrajectoryViewport::UpdatePreviewPoseFromScrubTime(const IAnimationProvider* InAnimationProvider, const IGameplayProvider* InGameplayProvider, double InTraceStartTime, double InTraceEndTime) const
+void SPreviewTrajectoryViewport::UpdatePreviewPoseFromScrubTime(double InTraceStartTime, double InTraceEndTime) const
 {
 	if (PreviewComponent->GetSkeletalMeshAsset() == nullptr || PreviewData->SourceIndex == INDEX_NONE)
 	{
@@ -277,17 +280,8 @@ void FPreviewTrajectoryViewportClient::Tick(float DeltaTime)
 			// Get data providers for traced objets.
 			if (const TraceServices::IAnalysisSession* Session = RewindDebugger->GetAnalysisSession())
 			{
-				const IGameplayProvider* GameplayProvider = Session->ReadProvider<IGameplayProvider>("GameplayProvider");
-				const IAnimationProvider* AnimationProvider = Session->ReadProvider<IAnimationProvider>("AnimationProvider");
-
-				// Preview traced skeletal mesh.
-				if (GameplayProvider && AnimationProvider)
-				{
-					// Start reading this trace time information.
-					TraceServices::FAnalysisSessionReadScope AnalysisSessionReadScope(*Session);
-					PreviewEditorViewport.Pin().Get()->UpdatePreviewPoseFromScrubTime(AnimationProvider, GameplayProvider, 0, RewindDebugger->CurrentTraceTime());
-				}
-			}	
+				PreviewEditorViewport.Pin().Get()->UpdatePreviewPoseFromScrubTime(0, RewindDebugger->CurrentTraceTime());
+			}
 		}
 	}
 
@@ -1162,9 +1156,6 @@ void SExportTrajectoriesWindow::UpdateAssetInfo(FSoftObjectPath&  OutSkeletonPat
 	
 	if (GameplayProvider && AnimationProvider)
 	{
-		// Start reading this trace time information.
-		TraceServices::FAnalysisSessionReadScope AnalysisSessionReadScope(*Session);
-
 		// Preview using skeletal mesh from current trace range.
 		if (PreviewData->OutputTrajectory.TraceInfo.Num() > 1)
 		{
@@ -1181,6 +1172,7 @@ void SExportTrajectoriesWindow::UpdateAssetInfo(FSoftObjectPath&  OutSkeletonPat
 				// Ensure our skeleton is valid.
 				if (SkeletalMeshInfo.SkeletonId != 0)
 				{
+					TraceServices::FProviderReadScopeLock GameplayProviderReadScope(*GameplayProvider);
 					const FObjectInfo& SkeletonObjectInfo = GameplayProvider->GetObjectInfo(SkeletalMeshInfo.SkeletonId);
 					OutSkeletonPath = FSoftObjectPath(SkeletonObjectInfo.PathName);
 					
@@ -1190,8 +1182,10 @@ void SExportTrajectoriesWindow::UpdateAssetInfo(FSoftObjectPath&  OutSkeletonPat
 		}
 
 		// Preview using initial traced skeletal mesh.
+		TraceServices::FProviderReadScopeLock AnimationProviderReadScope(*AnimationProvider);
 		if (const FSkeletalMeshInfo* SkeletalMeshInfo = AnimationProvider->FindSkeletalMeshInfo(ImmutableState.SkeletalMeshInfos[SelectedTrajectoryIndex].Id))
 		{
+			TraceServices::FProviderReadScopeLock GameplayProviderReadScope(*GameplayProvider);
 			if (const FObjectInfo* SkeletalMeshObjectInfo = GameplayProvider->FindObjectInfo(ImmutableState.SkeletalMeshInfos[SelectedTrajectoryIndex].Id))
 			{
 				OutSkeletalMeshPath = FSoftObjectPath(SkeletalMeshObjectInfo->PathName);
diff --git a/Engine/Plugins/Experimental/Animation/TrajectoryTools/Source/TrajectoryTools/Public/SExportTrajectoriesWindow.h b/Engine/Plugins/Experimental/Animation/TrajectoryTools/Source/TrajectoryTools/Public/SExportTrajectoriesWindow.h
--- a/Engine/Plugins/Experimental/Animation/TrajectoryTools/Source/TrajectoryTools/Public/SExportTrajectoriesWindow.h
+++ b/Engine/Plugins/Experimental/Animation/TrajectoryTools/Source/TrajectoryTools/Public/SExportTrajectoriesWindow.h
@@ -116,7 +116,7 @@ public:
 	/** End SWidget */
 
 	void ResetPreviewSkeletalMesh() const;
-	void UpdatePreviewPoseFromScrubTime(const IAnimationProvider* InAnimationProvider, const IGameplayProvider* InGameplayProvider, double InTraceStartTime, double InTraceEndTime) const;
+	void UpdatePreviewPoseFromScrubTime(double InTraceStartTime, double InTraceEndTime) const;
 	
 	TSharedPtr<FPreviewData> GetPreviewData() const { return PreviewData; }
 
