TrajectoryTools - fix for crash while changing levels because of invalid FRewindDebuggerTrajectory::WorldToVisualize

#[review]-52759288 [REDACTED_HANDLE]
#Synced-CL 52747405

#Synced-CL 52765189

diff --git a/Engine/Plugins/Experimental/Animation/TrajectoryTools/Source/TrajectoryTools/Private/TrajectoryRewindDebuggerExtension.cpp b/Engine/Plugins/Experimental/Animation/TrajectoryTools/Source/TrajectoryTools/Private/TrajectoryRewindDebuggerExtension.cpp
--- a/Engine/Plugins/Experimental/Animation/TrajectoryTools/Source/TrajectoryTools/Private/TrajectoryRewindDebuggerExtension.cpp
+++ b/Engine/Plugins/Experimental/Animation/TrajectoryTools/Source/TrajectoryTools/Private/TrajectoryRewindDebuggerExtension.cpp
@@ -278,7 +278,7 @@ void FRewindDebuggerTrajectory::EnsureDebugDrawDelegateExists()
 void FRewindDebuggerTrajectory::DebugDraw(UCanvas* InCanvas, APlayerController* InController)
 {
 	// @todo: Only draw object ids which are part of the target object.
-	if (IsValid(WorldToVisualize) && !State.Trajectories.IsEmpty())
+	if (WorldToVisualize.IsValid() && !State.Trajectories.IsEmpty())
 	{
 		for (int TrajectoryIndex = 0; TrajectoryIndex < State.Trajectories.Num(); ++TrajectoryIndex)
 		{
@@ -296,13 +296,13 @@ void FRewindDebuggerTrajectory::DebugDraw(UCanvas* InCanvas, APlayerController*
 					{
 						const FGameplayTrajectory::FSample& FutureSample = DebugTrajectoryPair.Samples[SampleIndex + 1];
 
-						DrawDebugLine(WorldToVisualize, Sample.Position, FutureSample.Position, DebugDrawColorForTrajectory);
+						DrawDebugLine(WorldToVisualize.Get(), Sample.Position, FutureSample.Position, DebugDrawColorForTrajectory);
 					}
 
 					// Draw sections at every second
 					if (SampleIndex % 20 == 0)
 					{
-						DrawDebugPoint(WorldToVisualize, Sample.Position, 5, DebugDrawColorForTrajectory);
+						DrawDebugPoint(WorldToVisualize.Get(), Sample.Position, 5, DebugDrawColorForTrajectory);
 					}
 				}
 			}
diff --git a/Engine/Plugins/Experimental/Animation/TrajectoryTools/Source/TrajectoryTools/Public/TrajectoryRewindDebuggerExtension.h b/Engine/Plugins/Experimental/Animation/TrajectoryTools/Source/TrajectoryTools/Public/TrajectoryRewindDebuggerExtension.h
--- a/Engine/Plugins/Experimental/Animation/TrajectoryTools/Source/TrajectoryTools/Public/TrajectoryRewindDebuggerExtension.h
+++ b/Engine/Plugins/Experimental/Animation/TrajectoryTools/Source/TrajectoryTools/Public/TrajectoryRewindDebuggerExtension.h
@@ -90,7 +90,7 @@ private:
 	
 	TArray<FColor> DebugDrawColors;
 	FDelegateHandle DebugDrawDelegateHandle;
-	UWorld* WorldToVisualize = nullptr;
+	TWeakObjectPtr<UWorld> WorldToVisualize;
 	TSharedPtr<SWindow> BakeOutWindow;
 	
 	void ShowBakeOutTrajectoryWindow();
